PR #1349: fix spelling error (Contributed by [REDACTED_PERSON])

diff --git a/Engine/Plugins/Experimental/HTML5Networking/README.md b/Engine/Plugins/Experimental/HTML5Networking/README.md
--- a/Engine/Plugins/Experimental/HTML5Networking/README.md
+++ b/Engine/Plugins/Experimental/HTML5Networking/README.md
@@ -1,7 +1,7 @@
 ### Unreal Engine 4 Networking over Websockets Plugin 
 
 - Provides websocket transport layer for unreal engine 4. 
-- Uses [libwebsockets](http://libwebsocket.org) for the server side and client side for non HTML5 clients.  
+- Uses [libwebsockets](http://libwebsockets.org) for the server side and client side for non HTML5 clients.  
 - HTML5 clients use emscripten's sockets abstraction. 
 
 #### How to
