Python API doc generation commandlet

The command-line arguments sent to the UE4 Editor to run the PythonOnlineDocs commandlet looks like this for the UE4 Editor with the standard Engine API documented:

-run=PythonOnlineDocs -NoShaderCompile -IncludeEngine -EnableAllPlugins -ExceptPlugins="MixedRealityCaptureFramework,Myo,Oodle,RenderDocPlugin,TesnorflowPlugin"

Enterprise version of command-line:

-run=PythonOnlineDocs -NoShaderCompile -IncludeEngine -IncludeEnterprise -EnableAllPlugins -ExceptPlugins="MixedRealityCaptureFramework,Myo,Oodle,RenderDocPlugin,TesnorflowPlugin"

Additional info can be found here:
Engine\Plugins\Experimental\PythonScriptPlugin\SphinxDocs\PythonAPI_docs_readme.txt

Python API docs
- changed generation to use PythonOnlineDocs commandlet
- added filtering of Python API documentation to optionally remove or include: internal (NotForLicencees, NoRedist, CarefullyRedist, EpicInternal, and confidential platforms), project, enterprise, engine
- added initial list of links to sections in index page
- changed order of sections to be: modules, native types, structs, classes, enums, delegates
- reduced introduction and pointed link to more extensive documentation done with WordPress site
- added header to match existing documentation pages (from Robb Surridge)
- removed Python doc generation settings from Python plugin settings
- many other misc updates

Engine commands
- added -NoShaderCompile global command-line option
- added UPackage version of  FSourceCodeNavigation::FindModulePath()
- added GetRestrictedSubdirectories() and IsRestrictedPath() to FPaths class

diff --git a/Engine/Plugins/Experimental/PythonScriptPlugin/Source/PythonScriptPlugin/Private/PyOnlineDocsWriter.cpp b/Engine/Plugins/Experimental/PythonScriptPlugin/Source/PythonScriptPlugin/Private/PyOnlineDocsWriter.cpp
--- a/Engine/Plugins/Experimental/PythonScriptPlugin/Source/PythonScriptPlugin/Private/PyOnlineDocsWriter.cpp
+++ b/Engine/Plugins/Experimental/PythonScriptPlugin/Source/PythonScriptPlugin/Private/PyOnlineDocsWriter.cpp
@@ -1,10 +1,11 @@
 // Copyright 1998-2018 Epic Games, Inc. All Rights Reserved.
 
 #include "PyOnlineDocsWriter.h"
-#include "PyUtil.h"
+#include "PythonScriptPlugin.h"
 #include "PyGenUtil.h"
 #include "HAL/FileManager.h"
 #include "Logging/MessageLog.h"
+#include "Misc/CommandLine.h"
 #include "Misc/FileHelper.h"
 #include "Misc/Paths.h"
 #include "Runtime/Launch/Resources/Version.h"
@@ -72,6 +73,11 @@ FString FPyOnlineDocsWriter::GetSourcePath() const
 	return GetSphinxDocsPath() / TEXT("source");
 }
 
+FString FPyOnlineDocsWriter::GetBuildPath() const
+{
+	return GetSphinxDocsPath() / TEXT("build");
+}
+
 FString FPyOnlineDocsWriter::GetTemplatePath() const
 {
 	return GetSourcePath() / TEXT("_templates");
@@ -123,13 +129,16 @@ void FPyOnlineDocsWriter::GenerateIndexFile()
 		Section->TypeNames.StableSort();
 	}
 
+	FString SectionList;
 	FString TableOfContents;
 
 	// Accumulate all the modules into the table of contents
 	if (Modules.Num() > 0)
 	{
+		SectionList += TEXT("* :ref:`Modules`") LINE_TERMINATOR;
+
 		TableOfContents +=
-			LINE_TERMINATOR
+			TEXT(".. _Modules:") LINE_TERMINATOR
 			TEXT(".. toctree::") LINE_TERMINATOR
 			TEXT("    :maxdepth: 1") LINE_TERMINATOR
 			TEXT("    :caption: Modules") LINE_TERMINATOR LINE_TERMINATOR;
@@ -143,8 +152,19 @@ void FPyOnlineDocsWriter::GenerateIndexFile()
 	// Accumulate all the classes for each section into the table of contents
 	for (const TSharedRef<FPyOnlineDocsSection>& Section : Sections)
 	{
+		FString SectionRef(Section->Name);
+		SectionRef.ReplaceInline(TEXT(" "), TEXT("-"));
+
+		SectionList += TEXT("* :ref:`");
+		SectionList += SectionRef;
+		SectionList += TEXT("`") LINE_TERMINATOR;
+
 		TableOfContents +=
 			LINE_TERMINATOR
+			TEXT(".. _");
+		TableOfContents += SectionRef;
+		TableOfContents +=
+			TEXT(":") LINE_TERMINATOR
 			TEXT(".. toctree::") LINE_TERMINATOR
 			TEXT("    :maxdepth: 1") LINE_TERMINATOR
 			TEXT("    :caption: ");
@@ -157,8 +177,11 @@ void FPyOnlineDocsWriter::GenerateIndexFile()
 		}
 	}
 
-	// Replace {{TableOfContents}} with actual list
+	// Replace {{SectionList}} with actual list
 	FString IndexText = IndexTemplate;
+	IndexText.ReplaceInline(TEXT("{{SectionList}}"), *SectionList, ESearchCase::CaseSensitive);
+
+	// Replace {{TableOfContents}} with actual list
 	IndexText.ReplaceInline(TEXT("{{TableOfContents}}"), *TableOfContents, ESearchCase::CaseSensitive);
 
 	// Save out index file
@@ -258,7 +281,7 @@ void FPyOnlineDocsWriter::GenerateClassFiles()
 
 void FPyOnlineDocsWriter::GenerateFiles(const FString& InPythonStubPath)
 {
-	UE_LOG(LogPython, Log, TEXT("Generating Python API online docs used by Sphinx to generate static HTML..."));
+	UE_LOG(LogPython, Display, TEXT("Generating Python API online docs used by Sphinx to generate static HTML..."));
 
 	// Copy generated unreal module stub file to PythonScriptPlugin/SphinxDocs/modules
 	{
@@ -283,10 +306,112 @@ void FPyOnlineDocsWriter::GenerateFiles(const FString& InPythonStubPath)
 	GenerateModuleFiles();
 	GenerateClassFiles();
 
-	UE_LOG(LogPython, Log, TEXT(
-		"... Finished generating Python API online docs.\n"
-		"In the OS command prompt in PythonPlugin/SphinxDocs call `sphinx-build -b html source/ build/` to generate the HTML."
-	));
+	UE_LOG(LogPython, Display, TEXT(
+		"  ... finished generating Sphinx files."));
+
+	FString Commandline = FCommandLine::Get();
+	bool bUseSphinx = Commandline.Contains(TEXT("-NoHTML")) == false;
+
+	if (bUseSphinx)
+	{
+		// Call Sphinx to generate online Python API docs. (Default)
+
+		// Running as internal Python calls on the version embedded in UE4 rather than as an
+		// executed external process since other installs, paths and environment variables may act
+		// in unexpected ways. Could potentially use Python C++ API calls rather than Python scripts,
+		// though this keeps it clear and if the calls evolve over time the vast number of examples
+		// online are in Python rather than C++.
+
+		// Update pip and then install Sphinx if needed. If Sphinx and its dependencies are already
+		// installed then it will determine that quickly and move on to using it.
+		// More info on using pip within Python here:
+		//   https://pip.pypa.io/en/stable/user_guide/#using-pip-from-your-program
+
+		FString PyCommandStr;
+		FString PythonPath = FPaths::ConvertRelativePathToFull(FPaths::EngineSourceDir()) / TEXT("ThirdParty/Python/Win64/python.exe");
+
+		PyCommandStr += TEXT(
+			"import sys\n"
+			"import subprocess\n"
+			"subprocess.check_call(['");
+		PyCommandStr += PythonPath;
+		PyCommandStr += TEXT(
+			"', '-m', 'pip', 'install', '-q', '-U', 'pip'])\n"
+			"subprocess.check_call(['");
+		PyCommandStr += PythonPath;
+		PyCommandStr +=	TEXT(
+			"', '-m', 'pip', 'install', '-q', '--no-warn-script-location', 'sphinx'])\n"
+			"import sphinx\n");
+
+		// Alternate technique calling pip as a Python command though above is recommended by pip.
+		//
+		//FString PyCommandStr = TEXT(
+		//	"import pip\n"
+		//	"pip.main(['install', 'sphinx'])\n"
+		//	"import sphinx\n");
+
+		// Un-import full unreal module so Sphinx will use generated stub version of unreal module
+		PyCommandStr += TEXT(
+			"del unreal\n"
+			"del sys.modules['unreal']\n");
+
+		// Add on Sphinx build command
+		PyCommandStr += TEXT("sphinx.build_main(['sphinx-build', '-b', 'html', '");
+		PyCommandStr += GetSourcePath();
+		PyCommandStr += TEXT("', '");
+		PyCommandStr += GetBuildPath();
+		PyCommandStr += TEXT("'])");
+
+		UE_LOG(LogPython, Display, TEXT(
+			"Calling Sphinx in PythonPlugin/SphinxDocs to generate the HTML...\n\n"
+			"%s\n\n"
+			"This can take a long time - 16+ minutes for full build on test system...\n"),
+			*PyCommandStr);
+
+		bool bLogSphinx = Commandline.Contains(TEXT("-HTMLLog"));
+		ELogVerbosity::Type OldVerbosity = LogPython.GetVerbosity();
+
+		if (!bLogSphinx)
+		{
+			// Disable Python logging (default)
+			LogPython.SetVerbosity(ELogVerbosity::NoLogging);
+		}
+
+		// Run the Python commands
+		bool PyRunSuccess = FPythonScriptPlugin::Get()->RunString(*PyCommandStr);
+
+		if (!bLogSphinx)
+		{
+			// Re-enable Python logging
+			LogPython.SetVerbosity(OldVerbosity);
+		}
+
+		// The running of the Python commands seem to think there are errors no matter what so not much use to make this notification.
+		//if (PyRunSuccess)
+		//{
+		//	UE_LOG(LogPython, Display, TEXT(
+		//		"\n"
+		//		"  There was an internal Python error while generating the docs, though it may not be an issue.\n\n"
+		//		"  You could call the commandlet again with the '-HTMLLog' option to see more information.\n"));
+		//}
+
+		UE_LOG(LogPython, Display, TEXT(
+			"  ... finished generating Python API online docs!\n\n"
+			"Find them in the following directory:\n"
+			"  %s\n\n"
+			"See additional instructions and information in:\n"
+			"  PythonScriptPlugin/SphinxDocs/PythonAPI_docs_readme.txt\n"),
+			*GetBuildPath());
+	}
+	else
+	{
+		// Prompt to manually call Sphinx to generate online Python API docs.
+		UE_LOG(LogPython, Display, TEXT(
+			"To build the Python API online docs manually follow the instructions in:\n"
+			"  PythonScriptPlugin/SphinxDocs/PythonAPI_docs_readme.txt\n\n"
+			"And then call:"
+			"  PythonScriptPlugin/SphinxDocs/sphinx-build -b html source/ build/"));
+	}
 }
 
 #endif	// WITH_PYTHON
diff --git a/Engine/Plugins/Experimental/PythonScriptPlugin/Source/PythonScriptPlugin/Private/PyOnlineDocsWriter.h b/Engine/Plugins/Experimental/PythonScriptPlugin/Source/PythonScriptPlugin/Private/PyOnlineDocsWriter.h
--- a/Engine/Plugins/Experimental/PythonScriptPlugin/Source/PythonScriptPlugin/Private/PyOnlineDocsWriter.h
+++ b/Engine/Plugins/Experimental/PythonScriptPlugin/Source/PythonScriptPlugin/Private/PyOnlineDocsWriter.h
@@ -3,9 +3,24 @@
 #pragma once
 
 #include "CoreMinimal.h"
+#include "Misc/EnumClassFlags.h"
 
 #if WITH_PYTHON
 
+/**
+ * Flags controlling which data is included in the Python API online docs.
+ */
+enum class EPyOnlineDocsFilterFlags : uint8
+{
+	IncludeNone = 0,
+	IncludeEngine = 1<<0,
+	IncludeEnterprise = 1<<1,
+	IncludeInternal = 1<<2,
+	IncludeProject = 1<<3,
+	IncludeAll = IncludeEngine | IncludeEnterprise | IncludeInternal | IncludeProject,
+};
+ENUM_CLASS_FLAGS(EPyOnlineDocsFilterFlags);
+
 /**
  * A single module in the Python API online docs.
  * Hosts a series of function names that belong to this module and will be used for indexing purposes later.
@@ -48,6 +63,9 @@ public:
 	/** Store class name in this section to generate files later. */
 	void AccumulateClass(const TCHAR* InTypeName);
 
+	//** Get name of section. */
+	const FString& GetName() const	{ return Name; }
+
 private:
 	/** Section name. */
 	FString Name;
@@ -76,6 +94,9 @@ public:
 	/** Get the directory for the Sphinx source files. */
 	FString GetSourcePath() const;
 
+	/** Get the directory for the Sphinx build files. */
+	FString GetBuildPath() const;
+
 	/** Get the directory for the Sphinx template files. */
 	FString GetTemplatePath() const;
 
diff --git a/Engine/Plugins/Experimental/PythonScriptPlugin/Source/PythonScriptPlugin/Private/PyWrapperBase.h b/Engine/Plugins/Experimental/PythonScriptPlugin/Source/PythonScriptPlugin/Private/PyWrapperBase.h
--- a/Engine/Plugins/Experimental/PythonScriptPlugin/Source/PythonScriptPlugin/Private/PyWrapperBase.h
+++ b/Engine/Plugins/Experimental/PythonScriptPlugin/Source/PythonScriptPlugin/Private/PyWrapperBase.h
@@ -70,6 +70,12 @@ struct FPyWrapperBaseMetaData
 	/** Get the ID associated with this meta-data type */
 	virtual FGuid GetTypeId() const = 0;
 
+	/** Get the reflection meta data type object associated with this wrapper type if there is one or nullptr if not. */
+	virtual const UField* GetMetaType() const
+	{
+		return nullptr;
+	}
+
 	/** Add object references from the given Python object to the given collector */
 	virtual void AddReferencedObjects(FPyWrapperBase* Instance, FReferenceCollector& Collector)
 	{
diff --git a/Engine/Plugins/Experimental/PythonScriptPlugin/Source/PythonScriptPlugin/Private/PyWrapperDelegate.h b/Engine/Plugins/Experimental/PythonScriptPlugin/Source/PythonScriptPlugin/Private/PyWrapperDelegate.h
--- a/Engine/Plugins/Experimental/PythonScriptPlugin/Source/PythonScriptPlugin/Private/PyWrapperDelegate.h
+++ b/Engine/Plugins/Experimental/PythonScriptPlugin/Source/PythonScriptPlugin/Private/PyWrapperDelegate.h
@@ -111,6 +111,12 @@ struct TPyWrapperDelegateMetaData : public FPyWrapperBaseMetaData
 		Collector.AddReferencedObject(PythonCallableForDelegateClass);
 	}
 
+	/** Get the reflection meta data type object associated with this wrapper type if there is one or nullptr if not. */
+	virtual const UField* GetMetaType() const override
+	{
+		return DelegateSignature.Func;
+	}
+
 	/** Unreal function representing the signature for the delegate */
 	PyGenUtil::FGeneratedWrappedFunction DelegateSignature;
 
diff --git a/Engine/Plugins/Experimental/PythonScriptPlugin/Source/PythonScriptPlugin/Private/PyWrapperEnum.h b/Engine/Plugins/Experimental/PythonScriptPlugin/Source/PythonScriptPlugin/Private/PyWrapperEnum.h
--- a/Engine/Plugins/Experimental/PythonScriptPlugin/Source/PythonScriptPlugin/Private/PyWrapperEnum.h
+++ b/Engine/Plugins/Experimental/PythonScriptPlugin/Source/PythonScriptPlugin/Private/PyWrapperEnum.h
@@ -85,6 +85,12 @@ struct FPyWrapperEnumMetaData : public FPyWrapperBaseMetaData
 	/** Check to see if the enum is finalized */
 	static bool IsEnumFinalized(FPyWrapperEnum* Instance);
 
+	/** Get the reflection meta data type object associated with this wrapper type if there is one or nullptr if not. */
+	virtual const UField* GetMetaType() const override
+	{
+		return Enum;
+	}
+
 	/** Unreal enum */
 	UEnum* Enum;
 
diff --git a/Engine/Plugins/Experimental/PythonScriptPlugin/Source/PythonScriptPlugin/Private/PyWrapperObject.h b/Engine/Plugins/Experimental/PythonScriptPlugin/Source/PythonScriptPlugin/Private/PyWrapperObject.h
--- a/Engine/Plugins/Experimental/PythonScriptPlugin/Source/PythonScriptPlugin/Private/PyWrapperObject.h
+++ b/Engine/Plugins/Experimental/PythonScriptPlugin/Source/PythonScriptPlugin/Private/PyWrapperObject.h
@@ -149,6 +149,12 @@ struct FPyWrapperObjectMetaData : public FPyWrapperBaseMetaData
 	/** Add object references from the given Python object to the given collector */
 	virtual void AddReferencedObjects(FPyWrapperBase* Instance, FReferenceCollector& Collector) override;
 
+	/** Get the reflection meta data type object associated with this wrapper type if there is one or nullptr if not. */
+	virtual const UField* GetMetaType() const override
+	{
+		return Class;
+	}
+
 	/** Unreal class */
 	UClass* Class;
 
diff --git a/Engine/Plugins/Experimental/PythonScriptPlugin/Source/PythonScriptPlugin/Private/PyWrapperStruct.h b/Engine/Plugins/Experimental/PythonScriptPlugin/Source/PythonScriptPlugin/Private/PyWrapperStruct.h
--- a/Engine/Plugins/Experimental/PythonScriptPlugin/Source/PythonScriptPlugin/Private/PyWrapperStruct.h
+++ b/Engine/Plugins/Experimental/PythonScriptPlugin/Source/PythonScriptPlugin/Private/PyWrapperStruct.h
@@ -331,6 +331,12 @@ struct FPyWrapperStructMetaData : public FPyWrapperBaseMetaData
 	/** Add object references from the given Python object to the given collector */
 	virtual void AddReferencedObjects(FPyWrapperBase* Instance, FReferenceCollector& Collector) override;
 
+	/** Get the reflection meta data type object associated with this wrapper type if there is one or nullptr if not. */
+	virtual const UField* GetMetaType() const override
+	{
+		return Struct;
+	}
+
 	/** Unreal struct */
 	UScriptStruct* Struct;
 
diff --git a/Engine/Plugins/Experimental/PythonScriptPlugin/Source/PythonScriptPlugin/Private/PyWrapperTypeRegistry.cpp b/Engine/Plugins/Experimental/PythonScriptPlugin/Source/PythonScriptPlugin/Private/PyWrapperTypeRegistry.cpp
--- a/Engine/Plugins/Experimental/PythonScriptPlugin/Source/PythonScriptPlugin/Private/PyWrapperTypeRegistry.cpp
+++ b/Engine/Plugins/Experimental/PythonScriptPlugin/Source/PythonScriptPlugin/Private/PyWrapperTypeRegistry.cpp
@@ -16,11 +16,11 @@
 #include "PyGIL.h"
 #include "PyCore.h"
 #include "PyFileWriter.h"
-#include "PyOnlineDocsWriter.h"
 #include "PythonScriptPluginSettings.h"
 #include "ProfilingDebugging/ScopedTimers.h"
 #include "Misc/Paths.h"
 #include "Misc/FileHelper.h"
+#include "SourceCodeNavigation.h"
 #include "UObject/UnrealType.h"
 #include "UObject/UObjectHash.h"
 #include "UObject/Class.h"
@@ -1934,8 +1934,10 @@ void FPyWrapperTypeRegistry::GatherWrappedTypesForPropertyReferences(const UProp
 	}
 }
 
-void FPyWrapperTypeRegistry::GenerateStubCodeForWrappedTypes() const
+void FPyWrapperTypeRegistry::GenerateStubCodeForWrappedTypes(const EPyOnlineDocsFilterFlags InDocGenFlags) const
 {
+	UE_LOG(LogPython, Display, TEXT("Generating Python API stub file..."));
+
 	FPyFileWriter PythonScript;
 
 	TUniquePtr<FPyOnlineDocsWriter> OnlineDocsWriter;
@@ -1946,15 +1948,15 @@ void FPyWrapperTypeRegistry::GenerateStubCodeForWrappedTypes() const
 	TSharedPtr<FPyOnlineDocsSection> OnlineDocsStructTypesSection;
 	TSharedPtr<FPyOnlineDocsSection> OnlineDocsClassTypesSection;
 
-	if (GetDefault<UPythonScriptPluginSettings>()->bGenerateOnlineDocs)
+	if (EnumHasAnyFlags(InDocGenFlags, EPyOnlineDocsFilterFlags::IncludeAll))
 	{
 		OnlineDocsWriter = MakeUnique<FPyOnlineDocsWriter>();
 		OnlineDocsUnrealModule = OnlineDocsWriter->CreateModule(TEXT("unreal"));
 		OnlineDocsNativeTypesSection = OnlineDocsWriter->CreateSection(TEXT("Native Types"));
-		OnlineDocsEnumTypesSection = OnlineDocsWriter->CreateSection(TEXT("Enum Types"));
-		OnlineDocsDelegateTypesSection = OnlineDocsWriter->CreateSection(TEXT("Delegate Types"));
 		OnlineDocsStructTypesSection = OnlineDocsWriter->CreateSection(TEXT("Struct Types"));
 		OnlineDocsClassTypesSection = OnlineDocsWriter->CreateSection(TEXT("Class Types"));
+		OnlineDocsEnumTypesSection = OnlineDocsWriter->CreateSection(TEXT("Enum Types"));
+		OnlineDocsDelegateTypesSection = OnlineDocsWriter->CreateSection(TEXT("Delegate Types"));
 	}
 
 	// Process additional Python files
@@ -2047,6 +2049,10 @@ void FPyWrapperTypeRegistry::GenerateStubCodeForWrappedTypes() const
 	}
 
 	// Process native glue code
+	UE_LOG(LogPython, Display, TEXT("  ...generating Python API: glue code"));
+	PythonScript.WriteLine(TEXT("##### Glue Code #####"));
+	PythonScript.WriteNewLine();
+
 	for (const PyGenUtil::FNativePythonModule& NativePythonModule : NativePythonModules)
 	{
 		for (PyMethodDef* MethodDef = NativePythonModule.PyModuleMethods; MethodDef && MethodDef->ml_name; ++MethodDef)
@@ -2072,11 +2078,83 @@ void FPyWrapperTypeRegistry::GenerateStubCodeForWrappedTypes() const
 	}
 
 	// Process generated glue code
-	auto ProcessWrappedDataArray = [this, &PythonScript](const TMap<FName, PyTypeObject*>& WrappedData, const TSharedPtr<FPyOnlineDocsSection>& OnlineDocsSection)
+	// Also excludes types that don't pass the filters specified in InDocGenFlags using the information about
+	// which module it came from and where that module exists on disk.
+	auto ProcessWrappedDataArray = [this, &PythonScript, InDocGenFlags](const TMap<FName, PyTypeObject*>& WrappedData, const TSharedPtr<FPyOnlineDocsSection>& OnlineDocsSection)
 	{
+		if (InDocGenFlags == EPyOnlineDocsFilterFlags::IncludeNone)
+		{
+			return;
+		}
+
+		UE_LOG(LogPython, Display, TEXT("  ...generating Python API: %s"), *OnlineDocsSection->GetName());
+		PythonScript.WriteLine(FString::Printf(TEXT("##### %s #####"), *OnlineDocsSection->GetName()));
+		PythonScript.WriteNewLine();
+		
+		FString ProjectTopDir;
+		if (FPaths::IsProjectFilePathSet())
+		{
+			ProjectTopDir / FPaths::GetCleanFilename(FPaths::ProjectDir());
+		}
+
 		for (const auto& WrappedDataPair : WrappedData)
 		{
 			TSharedPtr<PyGenUtil::FGeneratedWrappedType> GeneratedWrappedType = GeneratedWrappedTypes.FindRef(WrappedDataPair.Key);
+
+			if ((InDocGenFlags != EPyOnlineDocsFilterFlags::IncludeAll) && GeneratedWrappedType.IsValid())
+			{
+				const UField* MetaType = GeneratedWrappedType->MetaData->GetMetaType();
+
+				FString ModulePath;
+
+				if (MetaType)
+				{
+					FSourceCodeNavigation::FindModulePath(MetaType->GetTypedOuter<UPackage>(), ModulePath);
+				}
+
+				if (!ModulePath.IsEmpty())
+				{
+					// Is Project class?
+					if (!ProjectTopDir.IsEmpty()
+						&& (ModulePath.Contains(ProjectTopDir)))
+					{
+						// Optionally exclude Project classes
+						if (!EnumHasAnyFlags(InDocGenFlags, EPyOnlineDocsFilterFlags::IncludeProject))
+						{
+							continue;
+						}
+					}
+					// Is Enterprise class
+					else if (ModulePath.Contains(TEXT("/Enterprise/")))
+					{
+						// Optionally exclude Enterprise classes
+						if (!EnumHasAnyFlags(InDocGenFlags, EPyOnlineDocsFilterFlags::IncludeEnterprise))
+						{
+							continue;
+						}
+					}
+					// is internal class
+					else if (FPaths::IsRestrictedPath(ModulePath))
+					{
+						// Optionally exclude internal classes
+						if (!EnumHasAnyFlags(InDocGenFlags, EPyOnlineDocsFilterFlags::IncludeInternal))
+						{
+							continue;
+						}
+					}
+					// Everything else is considered an "Engine" class
+					else
+					{
+						// Optionally exclude engine classes
+						if (!EnumHasAnyFlags(InDocGenFlags, EPyOnlineDocsFilterFlags::IncludeEngine))
+						{
+							continue;
+						}
+					}
+				}
+				// else if cannot determine origin then include
+			}
+
 			GenerateStubCodeForWrappedType(WrappedDataPair.Value, GeneratedWrappedType.Get(), PythonScript, OnlineDocsSection.Get());
 		}
 	};
@@ -2087,13 +2165,18 @@ void FPyWrapperTypeRegistry::GenerateStubCodeForWrappedTypes() const
 	ProcessWrappedDataArray(PythonWrappedClasses, OnlineDocsClassTypesSection);
 
 	// Append any additional Python code now that all the reflected API has been exported
+	UE_LOG(LogPython, Display, TEXT("  ...generating Python API: additional code"));
+	PythonScript.WriteLine(TEXT("##### Additional Code #####"));
+	PythonScript.WriteNewLine();
+
 	for (const FString& AdditionalPythonLine : AdditionalPythonCode)
 	{
 		PythonScript.WriteLine(AdditionalPythonLine);
 	}
 
-	const FString PythonSourceFilename = FPaths::ProjectIntermediateDir() / TEXT("PythonStub") / TEXT("unreal.py");
+	const FString PythonSourceFilename = FPaths::ConvertRelativePathToFull(FPaths::ProjectIntermediateDir()) / TEXT("PythonStub") / TEXT("unreal.py");
 	PythonScript.SaveFile(*PythonSourceFilename);
+	UE_LOG(LogPython, Display, TEXT("  ...generated: %s"), *PythonSourceFilename);
 
 	if (OnlineDocsWriter.IsValid())
 	{
diff --git a/Engine/Plugins/Experimental/PythonScriptPlugin/Source/PythonScriptPlugin/Private/PyWrapperTypeRegistry.h b/Engine/Plugins/Experimental/PythonScriptPlugin/Source/PythonScriptPlugin/Private/PyWrapperTypeRegistry.h
--- a/Engine/Plugins/Experimental/PythonScriptPlugin/Source/PythonScriptPlugin/Private/PyWrapperTypeRegistry.h
+++ b/Engine/Plugins/Experimental/PythonScriptPlugin/Source/PythonScriptPlugin/Private/PyWrapperTypeRegistry.h
@@ -8,6 +8,7 @@
 #include "PyUtil.h"
 #include "PyGenUtil.h"
 #include "PyConversionMethod.h"
+#include "PyOnlineDocsWriter.h"
 #include "UObject/WeakObjectPtr.h"
 #include "Templates/Function.h"
 
@@ -30,7 +31,6 @@ class UPythonGeneratedStruct;
 class IPyWrapperInlineStructFactory;
 
 class FPyFileWriter;
-class FPyOnlineDocsSection;
 
 /** Type conversion for TPyWrapperTypeFactory */
 template <typename UnrealType, typename KeyType>
@@ -428,7 +428,7 @@ public:
 	PyTypeObject* GetWrappedDelegateType(const UFunction* InDelegateSignature) const;
 
 	/** Generate stub Python code for our wrapped types */
-	void GenerateStubCodeForWrappedTypes() const;
+	void GenerateStubCodeForWrappedTypes(const EPyOnlineDocsFilterFlags InDocGenFlags = EPyOnlineDocsFilterFlags::IncludeNone) const;
 
 private:
 	/** Gather any types referenced by the given property are still need to be wrapped for use in Python */
diff --git a/Engine/Plugins/Experimental/PythonScriptPlugin/Source/PythonScriptPlugin/Private/PythonOnlineDocsCommandlet.cpp b/Engine/Plugins/Experimental/PythonScriptPlugin/Source/PythonScriptPlugin/Private/PythonOnlineDocsCommandlet.cpp
new file mode 100644
--- /dev/null
+++ b/Engine/Plugins/Experimental/PythonScriptPlugin/Source/PythonScriptPlugin/Private/PythonOnlineDocsCommandlet.cpp
@@ -0,0 +1,56 @@
+// Copyright 1998-2018 Epic Games, Inc. All Rights Reserved.
+
+#include "PythonOnlineDocsCommandlet.h"
+#include "PyWrapperTypeRegistry.h"
+
+DEFINE_LOG_CATEGORY_STATIC(LogPythonOnlineDocsCommandlet, Log, All);
+
+UPythonOnlineDocsCommandlet::UPythonOnlineDocsCommandlet()
+{
+	IsServer = true;
+	IsClient = true;
+	IsEditor = true;
+	LogToConsole = false;
+	ShowErrorCount = false;
+}
+
+int32 UPythonOnlineDocsCommandlet::Main(const FString& Params)
+{
+	TArray<FString> Tokens;
+	TArray<FString> Switches;
+	TMap<FString, FString> ParamVals;
+	UCommandlet::ParseCommandLine(*Params, Tokens, Switches, ParamVals);
+
+#if WITH_PYTHON
+	EPyOnlineDocsFilterFlags DocGenFlags = EPyOnlineDocsFilterFlags::IncludeNone;
+
+	// Apply the filter flags from the command line
+#define APPLY_FILTER_FLAG(NAME)							\
+	if (Switches.Contains(TEXT(#NAME)))					\
+	{													\
+		DocGenFlags |= EPyOnlineDocsFilterFlags::NAME;	\
+	}
+	APPLY_FILTER_FLAG(IncludeEngine)
+	APPLY_FILTER_FLAG(IncludeEnterprise)
+	APPLY_FILTER_FLAG(IncludeInternal)
+	APPLY_FILTER_FLAG(IncludeProject)
+#undef APPLY_FILTER_FLAG
+
+	// If we weren't given any filter flags, include everything
+	if (DocGenFlags == EPyOnlineDocsFilterFlags::IncludeNone)
+	{
+		DocGenFlags = EPyOnlineDocsFilterFlags::IncludeAll;
+	}
+
+	UE_LOG(LogPythonOnlineDocsCommandlet, Display, TEXT("\n\nGenerating Python documentation..."));
+	FPyWrapperTypeRegistry::Get().GenerateStubCodeForWrappedTypes(DocGenFlags);
+
+#else	// WITH_PYTHON
+	
+	UE_LOG(LogPythonOnlineDocsCommandlet, Error, TEXT("Python docs cannot be generated as the plugin was built as a stub!"));
+	return -1;
+
+#endif	// WITH_PYTHON
+
+	return 0;
+}
diff --git a/Engine/Plugins/Experimental/PythonScriptPlugin/Source/PythonScriptPlugin/Private/PythonOnlineDocsCommandlet.h b/Engine/Plugins/Experimental/PythonScriptPlugin/Source/PythonScriptPlugin/Private/PythonOnlineDocsCommandlet.h
new file mode 100644
--- /dev/null
+++ b/Engine/Plugins/Experimental/PythonScriptPlugin/Source/PythonScriptPlugin/Private/PythonOnlineDocsCommandlet.h
@@ -0,0 +1,28 @@
+// Copyright 1998-2018 Epic Games, Inc. All Rights Reserved.
+
+#pragma once
+
+#include "Commandlets/Commandlet.h"
+#include "PythonOnlineDocsCommandlet.generated.h"
+
+/**
+ * Minimal commandlet to format and write Python API online docs.
+ * Can be passed the following flags to filter which types are included in the docs:
+ *	-IncludeEngine
+ *	-IncludeEnterprise
+ *	-IncludeInternal
+ *	-IncludeProject
+ */
+UCLASS()
+class UPythonOnlineDocsCommandlet : public UCommandlet
+{
+	GENERATED_BODY()
+
+public:
+	/** Default constructor. */
+	UPythonOnlineDocsCommandlet();
+
+	//~ Begin UCommandlet Interface
+	virtual int32 Main(const FString& Params) override;
+	//~ End UCommandlet Interface
+};
diff --git a/Engine/Plugins/Experimental/PythonScriptPlugin/Source/PythonScriptPlugin/Private/PythonScriptPluginSettings.h b/Engine/Plugins/Experimental/PythonScriptPlugin/Source/PythonScriptPlugin/Private/PythonScriptPluginSettings.h
--- a/Engine/Plugins/Experimental/PythonScriptPlugin/Source/PythonScriptPlugin/Private/PythonScriptPluginSettings.h
+++ b/Engine/Plugins/Experimental/PythonScriptPlugin/Source/PythonScriptPlugin/Private/PythonScriptPluginSettings.h
@@ -34,11 +34,4 @@ public:
 	/** Should Developer Mode be enabled on the Python interpreter (will enable extra warnings (eg, for deprecated code), and enable stub code generation for use with external IDEs). */
 	UPROPERTY(config, EditAnywhere, Category=Python, meta=(ConfigRestartRequired=true))
 	bool bDeveloperMode;
-
-	/**
-	* Should files used for Unreal Python API doc generation be created?
-	* [Also needs Developer Mode above enabled.]
-	*/
-	UPROPERTY(config, EditAnywhere, Category = Python, meta = (ConfigRestartRequired = true))
-	bool bGenerateOnlineDocs;
 };
diff --git a/Engine/Plugins/Experimental/PythonScriptPlugin/SphinxDocs/PythonAPI_docs_readme.txt b/Engine/Plugins/Experimental/PythonScriptPlugin/SphinxDocs/PythonAPI_docs_readme.txt
--- a/Engine/Plugins/Experimental/PythonScriptPlugin/SphinxDocs/PythonAPI_docs_readme.txt
+++ b/Engine/Plugins/Experimental/PythonScriptPlugin/SphinxDocs/PythonAPI_docs_readme.txt
@@ -1,31 +1,15 @@
-Unreal Python API HTML generation using Sphinx
+Unreal Python API HTML generation using the Python Online Docs Commandlet
 ===========================================================
 
-Ensure Python is installed
----------------------------------------
-
-Download from https://www.python.org/downloads/
-
-[Python 2.7.14 was used for testing and to make this guide since it matches the current needs of enterprise studios.]
+This commandlet can be run manually to generate the Python API HTML files locally or it can be set up on an automated server that periodically refreshes the docs online.
 
-You *may* need to set Python to use UTF-8 (if you have a code page set with "chcp 65001" for example). You can set it interactively on the command line:
-set PYTHONIOENCODING=UTF-8
-
-Or set it permanently in the environment variables.
 
-
-Install Sphinx
+Install UE4 and the Python plugin
 ---------------------------------------
 
-[Should only need to do this step once and then it is set up for repeated builds.]
-
-Install Sphinx using command prompt:
-pip install -U sphinx
-
-Confirm installation and version:
-sphinx-build --version
+Ensure the version of UE4 that you want to be working with is installed and the PythonScriptPlugin is enabled. Python 2.7 now comes auto-installed with UE4 and the doc generation commandlet will use that version so you don't need to install it yourself. All the other dependencies (such as the Sphinx Python doc generation tool) will automatically be downloaded and installed as needed when this commandlet is run.
 
-[sphinx-build 1.7.2 was used at the time of writing.]
+Note that the auto-install of Sphinx adds about 50MB.
 
 
 Ensure Sphinx template files are in place
@@ -33,20 +17,27 @@ Ensure Sphinx template files are in place
 
 [Should only need to do this step once and then it is set up for repeated builds.]
 
-Ensure the PythonScriptPlugin is synced from Perforce.
-
-The Sphinx template files used to for the doc generation are located here:
+The Sphinx template files used for the doc generation are located here:
 Engine\Plugins\Experimental\PythonScriptPlugin\SphinxDocs
 
-The two main template files are in SphinxDocs\source\_templates and they are used by the PythonScriptPlugin when creating files used by Sphinx:
+The template files are in SphinxDocs\source\_templates and they are used by the PythonScriptPlugin when creating files used by Sphinx:
+
+conf.py
+This is the settings file for the Sphinx doc generation.
+{{Version}} is replaced with the version of Unreal that is having docs generated.
 
 index.rst
-{{SectionsTOC}} is replaced with several tables of contents for all the Unreal Python classes and placed at SphinxDocs\source\index.rst
+{{SectionList}} is replaced with a list of all the different section categories for quick access.
+{{TableOfContents}} is replaced with several tables of contents for all the different categories of Unreal Python classes and placed at SphinxDocs\source\index.rst
+
+Module.rst
+{{Module}} is replaced with each exposed Python module and {{ModuleFunctions}} is replaced with all the module top level functions and data members.
+Currently the only module is "unreal".
 
 Class.rst
-{{Class}} is replaced with the name of each specific Unreal Python class and each separate file is named [ClassName].rst and placed in SphinxDocs\source\class\
+{{Class}} is replaced with the name of each specific Unreal Python class and each separate file is placed in SphinxDocs\source\class\ and named [ClassName].rst
 
-You can add additional .rst files for additional documentation, though be sure to reference them in the SphinxDocs\source\_templates\index.rst file table of contents.
+You can add additional .rst files for additional documentation - see SphinxDocs\source\introduction.rst. Be sure to reference them in the SphinxDocs\source\_templates\index.rst file table of contents. Also try to put as little as possible in supplementary docs since they will not be automatically undergo localization like the Wordpress files for the rest of the Python docs.
 
 
 You can find out more about Sphinx at http://www.sphinx-doc.org/en/master/
@@ -59,25 +50,77 @@ You can add reStructuredText syntax highlighting to Notepad++:
 [REDACTED_PROFILE_URL]/reStructuredText_NPP
 
 
-Generate Sphinx files from UE4 Editor
+Run the PythonOnlineDocs commandlet
 ---------------------------------------
 
-[This step needs to be done for each refresh of the Unreal Python API whenever reflected engine bindings are changed.]
+The command-line arguments sent to the UE4 Editor to run the PythonOnlineDocs commandlet looks like this for the UE4 Editor with the standard Engine API documented:
+
+-run=PythonOnlineDocs -NoShaderCompile -IncludeEngine -EnableAllPlugins -ExceptPlugins="MixedRealityCaptureFramework,Myo,Oodle,RenderDocPlugin,TesnorflowPlugin"
 
-Build an up-to-date version of the UE4 editor including the PythonScriptPlugin.
+Enterprise version of command-line:
 
-Run the UE4 editor and ensure the PythonScriptPlugin is enabled. Reload the UE4 editor if necessary.
+-run=PythonOnlineDocs -NoShaderCompile -IncludeEngine -IncludeEnterprise -EnableAllPlugins -ExceptPlugins="MixedRealityCaptureFramework,Myo,Oodle,RenderDocPlugin,TesnorflowPlugin"
 
-If there are any delayed dynamically loaded modules that you want to include for the Python API docs then take whatever steps necessary to ensure that they are loaded.
 
-Go to the settings dialog for the Python plugin and enable both "Developer Mode" and "Generate Online Docs".
+The stubs for the Python API will be generated first and then all the files needed for Sphinx to generate HTML files. The HTML and associated files once generated will be about 200MB and will be placed in:
+Engine\Plugins\Experimental\PythonScriptPlugin\SphinxDocs\build
 
-You will see "Restart required to apply new settings" so select "Restart Now".
+The files located in PythonScriptPlugin\SphinxDocs\build are relative and can be used without a web server. Just drag the index.htm file to your browser of choice. Even the search works locally via a generated JavaScript search index file.
 
-In the future, the generation process now will occur (after a small delay) whenever a module is loaded. This happens at the load of any project and at other points when a new module is needed such as when a project is initially loaded.
+The files can be zipped up and distributed to be used locally or placed on a web server.
 
-The stubs for the Python API will be generated first and then all the files needed for Sphinx to generate HTML files. They will be placed in:
-Engine\Plugins\Experimental\PythonScriptPlugin\SphinxDocs
+The objects.inv and the associated .doctrees folder is used by other Python API web sites for linking to this site. It is only needed when placed on an online web server and not needed when just using the files locally.
+
+
+PythonOnlineDocs commandlet options
+---------------------------------------
+
+Here are the different command-line options relevant for the PythonOnlineDocs commandlet.
+
+Filter options:
+-InclueEngine      : Include the standard Engine Python API. Needed for all filter options.
+-IncludeEnterprise : Include Enterprise Python API
+-IncludeProject    : Include Python API from any loaded project
+
+Other useful options:
+- HTMLLog          : Output additional Python info, warnings and errors including feedback from running Sphinx (which has many warnings that can be safely ignored)
+-NoShaderCompile   : Don't compile shaders (shaders aren't needed by the PythonOnlineDocs commandlet)
+- EnableAllPlugins : Include all the plugins known to the UE4 editor
+-ExceptPlugins=""  : Don't include specified plugins. Some plugins cause issues or aren't needed in the Python API docs.
+-NoHTML            : Don't run Sphinx, potentially run Sphinx manually to build the HTML files see "Unreal Python API HTML generation manually using Sphinx"
+
+
+Unreal Python API HTML generation manually using Sphinx
+===========================================================
+
+Use the same command-line as above with -NoHTML added:
+
+-run=PythonOnlineDocs -NoShaderCompile -IncludeEngine -NoHTML -EnableAllPlugins -ExceptPlugins="MixedRealityCaptureFramework,Myo,Oodle,RenderDocPlugin,TesnorflowPlugin"
+
+
+Ensure Python is installed
+---------------------------------------
+
+Download from https://www.python.org/downloads/
+
+[Python 2.7.14 was used for testing and to make this guide since it matches the current needs of enterprise studios.]
+
+You *may* need to set Python to use UTF-8 (if you have a code page set with "chcp 65001" for example). You can set it interactively on the command line:
+set PYTHONIOENCODING=UTF-8
+
+Or set it permanently in the environment variables.
+
+
+Install Sphinx
+---------------------------------------
+
+[Should only need to do this step once and then it is set up for repeated builds.]
+
+Install Sphinx using command prompt:
+pip install -U sphinx
+
+Confirm installation and version:
+sphinx-build --version
 
 
 Generate Python API HTML using Sphinx
@@ -96,16 +139,3 @@ Engine\Plugins\Experimental\PythonScriptPlugin\SphinxDocs\build
 
 Sphinx looks at the docstrings in the unreal.py stub file during generation looking for reStructuredText markup. Many of the docstrings have things that Sphinx *thinks* are markup, though they are usually badly formed markup since people didn't know of the reStructuredText conventions when they wrote them. Such as: indentations without line gaps before and after, identifiers that end in an underscore_, and many other things. As of writing this, there are 159 warnings during the Sphinx generation process. Most can be safely ignored since they are accidental markup in the docstrings. None of them currently prevent the HTML generation.
 
-
-Use the Python API HTML
----------------------------------------
-
-[This step needs to be done for each refresh of the Unreal Python API whenever reflected engine bindings are changed.]
-
-The files located in PythonScriptPlugin\SphinxDocs\build are relative and can be used without a web server. Just drag the index.htm file to your browser of choice. Even the search works locally via a generated JavaScript search index file.
-
-The files can be zipped up and distributed to be used locally or placed on a web server.
-
-The objects.inv and the associated .doctrees folder is used by other Python API web sites for linking to this site. It is only needed when placed on an online web server and not needed when just using the files locally.
-
-Enjoy!
diff --git a/Engine/Plugins/Experimental/PythonScriptPlugin/SphinxDocs/themes/ue4api/layout.html b/Engine/Plugins/Experimental/PythonScriptPlugin/SphinxDocs/themes/ue4api/layout.html
--- a/Engine/Plugins/Experimental/PythonScriptPlugin/SphinxDocs/themes/ue4api/layout.html
+++ b/Engine/Plugins/Experimental/PythonScriptPlugin/SphinxDocs/themes/ue4api/layout.html
@@ -11,3 +11,23 @@
 {% if theme_collapsiblesidebar|tobool %}
 {% set script_files = script_files + ['_static/sidebar.js'] %}
 {% endif %}
+
+{% block header %}
+<div id="head"></div>
+<script>
+$(document).ready(function () {
+    $.ajax({
+        url: "https://www.unrealengine.com/epic-games-header?useServerUrl=true&lang=en-US",
+        context: document.body
+    }).done(function (data) {
+        $("#head").html(data);
+    }).fail(function( jqXHR, textStatus, errorThrown ) {
+	});
+});
+
+window._egNavAfterConfigSet = function() {
+	_epicGamesNav.hideSignIn = true;
+	_epicGamesNav.hideLocale = true;
+};
+</script>
+{% endblock %}
\ No newline at end of file
diff --git a/Engine/Source/Editor/UnrealEd/Private/SourceCodeNavigation.cpp b/Engine/Source/Editor/UnrealEd/Private/SourceCodeNavigation.cpp
--- a/Engine/Source/Editor/UnrealEd/Private/SourceCodeNavigation.cpp
+++ b/Engine/Source/Editor/UnrealEd/Private/SourceCodeNavigation.cpp
@@ -1796,11 +1796,11 @@ FSourceCodeNavigation::FOnNewModuleAdded& FSourceCodeNavigation::AccessOnNewModu
 	return FSourceCodeNavigationImpl::Get().OnNewModuleAdded;
 }
 
-bool FSourceCodeNavigation::FindModulePath( const FString& ModuleName, FString &OutModulePath )
+bool FSourceCodeNavigation::FindModulePath( const FString& InModuleName, FString &OutModulePath )
 {
 	// Try to find a file matching the module name
 	const TArray<FString>& ModuleNames = GetSourceFileDatabase().GetModuleNames();
-	FString FindModuleSuffix = FString(TEXT("/")) + ModuleName + ".Build.cs";
+	FString FindModuleSuffix = FString(TEXT("/")) + InModuleName + ".Build.cs";
 	for (int32 Idx = 0; Idx < ModuleNames.Num(); Idx++)
 	{
 		if (ModuleNames[Idx].EndsWith(FindModuleSuffix))
@@ -1812,92 +1812,113 @@ bool FSourceCodeNavigation::FindModulePath( const FString& ModuleName, FString &
 	return false;
 }
 
-bool FSourceCodeNavigation::FindClassHeaderPath( const UField *Field, FString &OutClassHeaderPath )
+bool FSourceCodeNavigation::FindModulePath( const UPackage* InModulePackage, FString& OutClassHeaderPath )
 {
-	// Get the class package, and skip past the "/Script/" portion to get the module name
-	UPackage *ModulePackage = Field->GetTypedOuter<UPackage>();
-	FString ModulePackageName = ModulePackage->GetName();
+	if (!InModulePackage)
+	{
+		return false;
+	}
+
+	// Get the package name and skip past the "/Script/" portion to get the module name
+	FString ModulePackageName = InModulePackage->GetName();
 
 	int32 ModuleNameIdx;
-	if(ModulePackageName.FindLastChar(TEXT('/'), ModuleNameIdx))
+
+	if (ModulePackageName.FindLastChar(TEXT('/'), ModuleNameIdx))
 	{
 		// Find the base path for the module
-		FString ModuleBasePath;
-		if(FSourceCodeNavigation::FindModulePath(*ModulePackageName + ModuleNameIdx + 1, ModuleBasePath))
+		return FindModulePath(*ModulePackageName + ModuleNameIdx + 1, OutClassHeaderPath);
+	}
+
+	return false;
+}
+
+bool FSourceCodeNavigation::FindClassHeaderPath( const UField* InField, FString &OutClassHeaderPath )
+{
+	if (!InField)
+	{
+		return false;
+	}
+
+	// Get the class package, and skip past the "/Script/" portion to get the module name
+	UPackage* ModulePackage = InField->GetTypedOuter<UPackage>();
+
+	// Find the base path for the module
+	FString ModuleBasePath;
+	if(FSourceCodeNavigation::FindModulePath(ModulePackage, ModuleBasePath))
+	{
+		// Get the metadata for the class path relative to the module base
+		const FString& ModuleRelativePath = ModulePackage->GetMetaData()->GetValue(InField, TEXT("ModuleRelativePath"));
+		if(ModuleRelativePath.Len() > 0)
 		{
-			// Get the metadata for the class path relative to the module base
-			const FString& ModuleRelativePath = ModulePackage->GetMetaData()->GetValue(Field, TEXT("ModuleRelativePath"));
-			if(ModuleRelativePath.Len() > 0)
-			{
-				OutClassHeaderPath = ModuleBasePath / ModuleRelativePath;
-				return true;
-			}
+			OutClassHeaderPath = ModuleBasePath / ModuleRelativePath;
+			return true;
 		}
 	}
 	return false;
 }
 
-bool FSourceCodeNavigation::FindClassSourcePath( const UField *Field, FString &OutClassSourcePath )
+bool FSourceCodeNavigation::FindClassSourcePath( const UField* InField, FString &OutClassSourcePath )
 {
+	if (!InField)
+	{
+		return false;
+	}
+
 	// Get the class package, and skip past the "/Script/" portion to get the module name
-	UPackage *ModulePackage = Field->GetTypedOuter<UPackage>();
-	FString ModulePackageName = ModulePackage->GetName();
+	UPackage *ModulePackage = InField->GetTypedOuter<UPackage>();
 
-	int32 ModuleNameIdx;
-	if(ModulePackageName.FindLastChar(TEXT('/'), ModuleNameIdx))
+	// Find the base path for the module
+	FString ModuleBasePath;
+	if (FSourceCodeNavigation::FindModulePath(ModulePackage, ModuleBasePath))
 	{
-		// Find the base path for the module
-		FString ModuleBasePath;
-		if(FSourceCodeNavigation::FindModulePath(*ModulePackageName + ModuleNameIdx + 1, ModuleBasePath))
+		// Get the metadata for the class path relative to the module base
+		// Given this we can try and find the corresponding .cpp file
+		const FString& ModuleRelativePath = ModulePackage->GetMetaData()->GetValue(InField, TEXT("ModuleRelativePath"));
+		if(ModuleRelativePath.Len() > 0)
 		{
-			// Get the metadata for the class path relative to the module base
-			// Given this we can try and find the corresponding .cpp file
-			const FString& ModuleRelativePath = ModulePackage->GetMetaData()->GetValue(Field, TEXT("ModuleRelativePath"));
-			if(ModuleRelativePath.Len() > 0)
-			{
-				const FString PotentialCppLeafname = FPaths::GetBaseFilename(ModuleRelativePath) + TEXT(".cpp");
-				FString PotentialCppFilename = ModuleBasePath / FPaths::GetPath(ModuleRelativePath) / PotentialCppLeafname;
+			const FString PotentialCppLeafname = FPaths::GetBaseFilename(ModuleRelativePath) + TEXT(".cpp");
+			FString PotentialCppFilename = ModuleBasePath / FPaths::GetPath(ModuleRelativePath) / PotentialCppLeafname;
 
-				// Is the .cpp file in the same folder as the header file?
-				if(FPaths::FileExists(PotentialCppFilename))
-				{
-					OutClassSourcePath = PotentialCppFilename;
-					return true;
-				}
+			// Is the .cpp file in the same folder as the header file?
+			if(FPaths::FileExists(PotentialCppFilename))
+			{
+				OutClassSourcePath = PotentialCppFilename;
+				return true;
+			}
 
-				const FString PublicPath = ModuleBasePath / "Public" / "";		// Ensure trailing /
-				const FString PrivatePath = ModuleBasePath / "Private" / "";	// Ensure trailing /
-				const FString ClassesPath = ModuleBasePath / "Classes" / "";	// Ensure trailing /
+			const FString PublicPath = ModuleBasePath / "Public" / "";		// Ensure trailing /
+			const FString PrivatePath = ModuleBasePath / "Private" / "";	// Ensure trailing /
+			const FString ClassesPath = ModuleBasePath / "Classes" / "";	// Ensure trailing /
 
-				// If the path starts with Public or Classes, try swapping those out with Private
-				if(PotentialCppFilename.StartsWith(PublicPath))
-				{
-					PotentialCppFilename.ReplaceInline(*PublicPath, *PrivatePath);
-				}
-				else if(PotentialCppFilename.StartsWith(ClassesPath))
-				{
-					PotentialCppFilename.ReplaceInline(*ClassesPath, *PrivatePath);
-				}
-				else
-				{
-					PotentialCppFilename.Empty();
-				}
-				if(!PotentialCppFilename.IsEmpty() && FPaths::FileExists(PotentialCppFilename))
-				{
-					OutClassSourcePath = PotentialCppFilename;
-					return true;
-				}
+			// If the path starts with Public or Classes, try swapping those out with Private
+			if(PotentialCppFilename.StartsWith(PublicPath))
+			{
+				PotentialCppFilename.ReplaceInline(*PublicPath, *PrivatePath);
+			}
+			else if(PotentialCppFilename.StartsWith(ClassesPath))
+			{
+				PotentialCppFilename.ReplaceInline(*ClassesPath, *PrivatePath);
+			}
+			else
+			{
+				PotentialCppFilename.Empty();
+			}
+			if(!PotentialCppFilename.IsEmpty() && FPaths::FileExists(PotentialCppFilename))
+			{
+				OutClassSourcePath = PotentialCppFilename;
+				return true;
+			}
 
-				// Still no luck, try and search for the file on the filesystem
-				TArray<FString> Filenames;
-				IFileManager::Get().FindFilesRecursive(Filenames, *ModuleBasePath, *PotentialCppLeafname, true, false, false);
+			// Still no luck, try and search for the file on the filesystem
+			TArray<FString> Filenames;
+			IFileManager::Get().FindFilesRecursive(Filenames, *ModuleBasePath, *PotentialCppLeafname, true, false, false);
 				
-				if(Filenames.Num() > 0)
-				{
-					// Assume it's the first match (we should really only find a single file with a given name within a project anyway)
-					OutClassSourcePath = Filenames[0];
-					return true;
-				}
+			if(Filenames.Num() > 0)
+			{
+				// Assume it's the first match (we should really only find a single file with a given name within a project anyway)
+				OutClassSourcePath = Filenames[0];
+				return true;
 			}
 		}
 	}
diff --git a/Engine/Source/Editor/UnrealEd/Public/SourceCodeNavigation.h b/Engine/Source/Editor/UnrealEd/Public/SourceCodeNavigation.h
--- a/Engine/Source/Editor/UnrealEd/Public/SourceCodeNavigation.h
+++ b/Engine/Source/Editor/UnrealEd/Public/SourceCodeNavigation.h
@@ -269,13 +269,16 @@ public:
 	UNREALED_API static bool IsCompilerAvailable() { return bCachedIsCompilerAvailable; }
 
 	/** Finds the base directory for a given module name. Does not rely on symbols; finds matching .build.cs files. */
-	UNREALED_API static bool FindModulePath( const FString& ModuleName, FString &OutModulePath );
+	UNREALED_API static bool FindModulePath( const FString& InModuleName, FString& OutModulePath );
+
+	/** Finds the base directory for a given module. Does not rely on symbols; finds matching .build.cs files. */
+	UNREALED_API static bool FindModulePath( const UPackage* InModulePackage, FString& OutModulePath );
 
 	/** Finds the path to a given class header. Does not rely on symbols; finds matching .build.cs files. */
-	UNREALED_API static bool FindClassHeaderPath( const UField *Field, FString &OutClassHeaderPath );
+	UNREALED_API static bool FindClassHeaderPath( const UField* InField, FString &OutClassHeaderPath );
 
 	/** Finds the path to a given class source. Does not rely on symbols; finds matching .build.cs files. */
-	UNREALED_API static bool FindClassSourcePath( const UField *Field, FString &OutClassSourcePath );
+	UNREALED_API static bool FindClassSourcePath( const UField* InField, FString &OutClassSourcePath );
 
 	/** Opens a single source file */
 	UNREALED_API static bool OpenSourceFile( const FString& AbsoluteSourcePath, int32 LineNumber = 0, int32 ColumnNumber = 0 );
diff --git a/Engine/Source/Runtime/Core/Private/Misc/Paths.cpp b/Engine/Source/Runtime/Core/Private/Misc/Paths.cpp
--- a/Engine/Source/Runtime/Core/Private/Misc/Paths.cpp
+++ b/Engine/Source/Runtime/Core/Private/Misc/Paths.cpp
@@ -14,6 +14,7 @@
 #include "Misc/Guid.h"
 #include "Misc/ConfigCacheIni.h"
 #include "Misc/App.h"
+#include "Misc/DataDrivenPlatformInfoRegistry.h"
 #include "Misc/EngineVersion.h"
 
 DEFINE_LOG_CATEGORY_STATIC(LogPaths, Log, All);
@@ -548,6 +549,66 @@ const TArray<FString>& FPaths::GetGameLocalizationPaths()
 	return Results;
 }
 
+const TArray<FString>& FPaths::GetRestrictedFolderNames()
+{
+	static bool bBuiltArray = false;
+	static TArray<FString> RestrictedFolderNames;
+
+	if (!bBuiltArray)
+	{
+		RestrictedFolderNames.Add(TEXT("NotForLicencees"));
+		RestrictedFolderNames.Add(TEXT("NoRedist"));
+		RestrictedFolderNames.Add(TEXT("CarefullyRedist"));
+		RestrictedFolderNames.Add(TEXT("EpicInternal"));
+
+		// Add confidential platforms
+		for (const FString& PlatformStr : FDataDrivenPlatformInfoRegistry::GetConfidentialPlatforms())
+		{
+			RestrictedFolderNames.Add(PlatformStr);
+		}
+
+		bBuiltArray = true;
+	}
+
+	return RestrictedFolderNames;
+}
+
+bool FPaths::IsRestrictedPath(const FString& InPath)
+{
+	static bool bBuiltArray = false;
+	static TArray<FString> RestrictedSlashedFolders;
+
+	if (!bBuiltArray)
+	{
+		// Add leading and trailing slashes to restricted folder names.
+		FString LeadingSlash(TEXT("/"));
+
+		for (const FString& FolderStr : GetRestrictedFolderNames())
+		{
+			RestrictedSlashedFolders.Add(LeadingSlash + FolderStr + TEXT('/'));
+		}
+
+		bBuiltArray = true;
+	}
+
+	// Normalize path
+	FString NormalizedPath(InPath);
+	NormalizeFilename(NormalizedPath);
+
+	// Ensure trailing forward slash
+	NormalizedPath /= FString();
+
+	for (const FString& SubDir : RestrictedSlashedFolders)
+	{
+		if (NormalizedPath.Contains(SubDir))
+		{
+			return true;
+		}
+	}
+
+	return false;
+}
+
 FString FPaths::GameAgnosticSavedDir()
 {
 	return EngineSavedDir();
diff --git a/Engine/Source/Runtime/Core/Public/Misc/Paths.h b/Engine/Source/Runtime/Core/Public/Misc/Paths.h
--- a/Engine/Source/Runtime/Core/Public/Misc/Paths.h
+++ b/Engine/Source/Runtime/Core/Public/Misc/Paths.h
@@ -317,7 +317,7 @@ public:
 	 */
 	static const TArray<FString>& GetPropertyNameLocalizationPaths();
 
-		/** 
+	/** 
 	 * Returns a list of tool tip localization paths
 	 */
 	static const TArray<FString>& GetToolTipLocalizationPaths();
@@ -327,6 +327,16 @@ public:
 	 */
 	static const TArray<FString>& GetGameLocalizationPaths();
 
+	/** 
+	 * Returns a list of restricted/internal folder names (without any slashes) which may be tested against full paths to determine if a path is restricted or not.
+	 */
+	static const TArray<FString>& GetRestrictedFolderNames();
+
+	/** 
+	 * Determines if supplied path uses a restricted/internal subdirectory.	Note that slashes are normalized and character case is ignored for the comparison.
+	 */
+	static bool IsRestrictedPath(const FString& InPath);
+
 	/**
 	 * Returns the saved directory that is not game specific. This is usually the same as
 	 * EngineSavedDir().
diff --git a/Engine/Source/Runtime/Engine/Classes/Commandlets/Commandlet.h b/Engine/Source/Runtime/Engine/Classes/Commandlets/Commandlet.h
--- a/Engine/Source/Runtime/Engine/Classes/Commandlets/Commandlet.h
+++ b/Engine/Source/Runtime/Engine/Classes/Commandlets/Commandlet.h
@@ -24,6 +24,9 @@
  *
  * Commandlets are executed in a "raw" environment, in which the game isn't
  * loaded, the client code isn't loaded, no levels are loaded, and no actors exist.
+ * 
+ * To disable shader compiling during the run of the commandlet add "-NoShaderCompile" to the commandline.
+ * This would be added as data member setting except that the shader compiler is initialized before a commandlet is created so it cannot be queried soon enough.
  */
 
 #pragma once
@@ -79,12 +82,11 @@ class UCommandlet : public UObject
 	/** Whether to show standard error and warning count on exit */
 	UPROPERTY()
 	uint32 ShowErrorCount:1;
-
+  
 	/** Whether to show cooking progress on tick */
 	UPROPERTY()
 	uint32 ShowProgress:1;
 
-
 	/**
 	 * Entry point for your commandlet
 	 *
diff --git a/Engine/Source/Runtime/Launch/Private/LaunchEngineLoop.cpp b/Engine/Source/Runtime/Launch/Private/LaunchEngineLoop.cpp
--- a/Engine/Source/Runtime/Launch/Private/LaunchEngineLoop.cpp
+++ b/Engine/Source/Runtime/Launch/Private/LaunchEngineLoop.cpp
@@ -1814,7 +1814,10 @@ int32 FEngineLoop::PreInit(const TCHAR* CmdLine)
 
 	FShaderCache::LoadBinaryCache();
 
-	if (!FPlatformProperties::RequiresCookedData())
+	FString Commandline = FCommandLine::Get();
+	bool EnableShaderCompile = !FParse::Param(*Commandline, TEXT("NoShaderCompile"));
+
+	if (EnableShaderCompile && !FPlatformProperties::RequiresCookedData())
 	{
 		check(!GShaderCompilingManager);
 		GShaderCompilingManager = new FShaderCompilingManager();
@@ -1836,10 +1839,10 @@ int32 FEngineLoop::PreInit(const TCHAR* CmdLine)
 		// Load the global shaders.
 		// if (!IsRunningCommandlet())
 		// hack: don't load global shaders if we are cooking we will load the shaders for the correct platform later
-		FString Commandline = FCommandLine::Get();
-		if (!IsRunningDedicatedServer() &&
-			Commandline.Contains(TEXT("cookcommandlet")) == false &&
-			Commandline.Contains(TEXT("run=cook")) == false )
+		if (EnableShaderCompile &&
+				!IsRunningDedicatedServer() &&
+				Commandline.Contains(TEXT("cookcommandlet")) == false &&
+				Commandline.Contains(TEXT("run=cook")) == false )
 		// if (FParse::Param(FCommandLine::Get(), TEXT("Multiprocess")) == false)
 		{
 			CompileGlobalShaderMap(false);
