HTML5 - documentation clarification when using websocket netdriver

#rnx

diff --git a/Engine/Plugins/Experimental/HTML5Networking/README.md b/Engine/Plugins/Experimental/HTML5Networking/README.md
--- a/Engine/Plugins/Experimental/HTML5Networking/README.md
+++ b/Engine/Plugins/Experimental/HTML5Networking/README.md
@@ -1,12 +1,12 @@
-### Unreal Engine 4 Networking over Websockets Plugin 
+### Unreal Engine 4 Networking over Websockets Plugin
 
-- Provides websocket transport layer for unreal engine 4. 
-- Uses [libwebsockets](http://libwebsockets.org) for the server side and client side for non HTML5 clients.  
-- HTML5 clients use emscripten's sockets abstraction. 
+- Provides websocket transport layer for unreal engine 4.
+- Uses [libwebsockets](http://libwebsockets.org) for the server side and client side for non HTML5 clients.
+- HTML5 clients use emscripten's sockets abstraction.
 
 #### How to
 
-- Clone this in Engine\Plugins\Experimental directory. 
+- Clone this in Engine\Plugins\Experimental directory.
 - Add the following section in `BaseEngine.ini`
 ```
 [/Script/HTML5Networking.WebSocketNetDriver]
@@ -27,16 +27,16 @@ WebSocketPort=8889
 NetConnectionClassName="/Script/HTML5Networking.WebSocketConnection"
 MaxPortCountToTry=512
 ```
-In section [/Script/Engine.Engine] disable comment out NetDriverDefinitions and add 
+In section [/Script/Engine.Engine] disable (by commenting out) ALL existing NetDriverDefinitions entries and then add:
 ```
 NetDriverDefinitions=(DefName="GameNetDriver",DriverClassName="/Script/HTML5Networking.WebSocketNetDriver",DriverClassNameFallback="/Script/HTML5Networking.IpNetDriver")
 ```
-To enable this Net Driver. 
+To enable this Net Driver.
 
 Build! and follow existing Unreal Networking documentation to setup servers/clients.
 
-#### Issues/Todo 
+#### Issues/Todo
 
-Disconnect events on client or server side are not handled properly yet 
+Disconnect events on client or server side are not handled properly yet
 
-Copyright 2015-2019 Epic Games. 
+Copyright 2015-2019 Epic Games.
