Fixed some issues that caused the unreal.py stub file to have syntax errors

We require this file to be error free in order to generate the Python API docs

#rnx

diff --git a/Engine/Plugins/Experimental/PythonScriptPlugin/SphinxDocs/PythonAPI_docs_readme.txt b/Engine/Plugins/Experimental/PythonScriptPlugin/SphinxDocs/PythonAPI_docs_readme.txt
--- a/Engine/Plugins/Experimental/PythonScriptPlugin/SphinxDocs/PythonAPI_docs_readme.txt
+++ b/Engine/Plugins/Experimental/PythonScriptPlugin/SphinxDocs/PythonAPI_docs_readme.txt
@@ -95,7 +95,7 @@ Unreal Python API HTML generation manually using Sphinx
 
 Use the same command-line as above with -NoHTML added:
 
--run=PythonOnlineDocs -NoShaderCompile -IncludeEngine -NoHTML -EnableAllPlugins -ExceptPlugins="MixedRealityCaptureFramework,Oodle,RenderDocPlugin,OculusAvatar"
+-run=PythonOnlineDocs -NoShaderCompile -IncludeEngine -NoHTML -EnableAllPlugins -ExceptPlugins="MixedRealityCaptureFramework,Oodle,RenderDocPlugin,OculusAvatar,VoiceChat"
 
 
 Ensure Python is installed
diff --git a/Engine/Source/Runtime/Engine/Classes/Engine/EngineBaseTypes.h b/Engine/Source/Runtime/Engine/Classes/Engine/EngineBaseTypes.h
--- a/Engine/Source/Runtime/Engine/Classes/Engine/EngineBaseTypes.h
+++ b/Engine/Source/Runtime/Engine/Classes/Engine/EngineBaseTypes.h
@@ -970,7 +970,7 @@ enum EViewModeIndex
 
 	VMI_Max UMETA(Hidden),
 
-	VMI_Unknown = 255 UMETA(Hidden),
+	VMI_Unknown = 255 UMETA(DisplayName = "Unknown"),
 };
 
 
