Replaced UE4 references in the Python documentation

diff --git a/Engine/Plugins/Experimental/PythonScriptPlugin/SphinxDocs/PythonAPI_docs_readme.txt b/Engine/Plugins/Experimental/PythonScriptPlugin/SphinxDocs/PythonAPI_docs_readme.txt
--- a/Engine/Plugins/Experimental/PythonScriptPlugin/SphinxDocs/PythonAPI_docs_readme.txt
+++ b/Engine/Plugins/Experimental/PythonScriptPlugin/SphinxDocs/PythonAPI_docs_readme.txt
@@ -4,10 +4,10 @@ Unreal Python API HTML generation using the Python Online Docs Commandlet
 This commandlet can be run manually to generate the Python API HTML files locally or it can be set up on an automated server that periodically refreshes the docs online.
 
 
-Install UE4 and the Python plugin
+Install Unreal Engine and the Python plugin
 ---------------------------------------
 
-Ensure the version of UE4 that you want to be working with is installed and the PythonScriptPlugin is enabled. Python 2.7 now comes auto-installed with UE4 and the doc generation commandlet will use that version so you don't need to install it yourself. All the other dependencies (such as the Sphinx Python doc generation tool) will automatically be downloaded and installed as needed when this commandlet is run.
+Ensure the version of Unreal Engine that you want to be working with is installed and the PythonScriptPlugin is enabled. Python 3.7 now comes auto-installed with Unreal Engine and the doc generation commandlet will use that version so you don't need to install it yourself. All the other dependencies (such as the Sphinx Python doc generation tool) will automatically be downloaded and installed as needed when this commandlet is run.
 
 Note that the auto-install of Sphinx adds about 50MB.
 
@@ -53,7 +53,7 @@ [REDACTED_PROFILE_URL]/reStructuredText_NPP
 Run the PythonOnlineDocs commandlet
 ---------------------------------------
 
-The command-line arguments sent to the UE4 Editor to run the PythonOnlineDocs commandlet looks like this for the UE4 Editor with the standard Engine API documented:
+The command-line arguments sent to the Unreal Engine Editor to run the PythonOnlineDocs commandlet looks like this for the Unreal Engine Editor with the standard Engine API documented:
 
 -run=PythonOnlineDocs -NoShaderCompile -IncludeEngine -EnableAllPlugins
 
@@ -85,7 +85,7 @@ Filter options:
 Other useful options:
 -HTMLLog           : Output additional Python info, warnings and errors including feedback from running Sphinx (which has many warnings that can be safely ignored)
 -NoShaderCompile   : Don't compile shaders (shaders aren't needed by the PythonOnlineDocs commandlet)
--EnableAllPlugins  : Include all the plugins known to the UE4 editor
+-EnableAllPlugins  : Include all the plugins known to the Unreal Engine editor
 -ExceptPlugins=""  : Don't include specified plugins. Some plugins cause issues or aren't needed in the Python API docs.
 -NoHTML            : Don't run Sphinx, potentially run Sphinx manually to build the HTML files see "Unreal Python API HTML generation manually using Sphinx"
 
diff --git a/Engine/Plugins/Experimental/PythonScriptPlugin/SphinxDocs/source/_templates/conf.py b/Engine/Plugins/Experimental/PythonScriptPlugin/SphinxDocs/source/_templates/conf.py
--- a/Engine/Plugins/Experimental/PythonScriptPlugin/SphinxDocs/source/_templates/conf.py
+++ b/Engine/Plugins/Experimental/PythonScriptPlugin/SphinxDocs/source/_templates/conf.py
@@ -87,7 +87,7 @@ pygments_style = 'sphinx'
 html_theme_path = ['../themes']
 
 # The theme to use for HTML and HTML Help pages.
-html_theme = 'ue4api'
+html_theme = 'ue_api'
 
 # Theme options are theme-specific and customize the look and feel of a theme.
 html_theme_options = {
diff --git a/Engine/Plugins/Experimental/PythonScriptPlugin/SphinxDocs/themes/ue4api/layout.html b/Engine/Plugins/Experimental/PythonScriptPlugin/SphinxDocs/themes/ue4api/layout.html
deleted file mode 100644
diff --git a/Engine/Plugins/Experimental/PythonScriptPlugin/SphinxDocs/themes/ue_api/layout.html b/Engine/Plugins/Experimental/PythonScriptPlugin/SphinxDocs/themes/ue_api/layout.html
new file mode 100644
--- /dev/null
+++ b/Engine/Plugins/Experimental/PythonScriptPlugin/SphinxDocs/themes/ue_api/layout.html
@@ -0,0 +1,35 @@
+{#
+    ue_api/layout.html
+    ~~~~~~~~~~~~~~~~~~~
+
+    Sphinx layout template for the Unreal Engine API theme.
+    
+    Copyright Epic Games, Inc. All Rights Reserved.
+#}
+{%- extends "classic/layout.html" %}
+
+{% if theme_collapsiblesidebar|tobool %}
+{% set script_files = script_files + ['_static/sidebar.js'] %}
+{% endif %}
+
+{% block header %}
+<script>window.__tracking_base = 'https://tracking.unrealengine.com';</script>
+<script src="https://tracking.unrealengine.com/tracking.js" async></script>
+<div id="head"></div>
+<script>
+$(document).ready(function () {
+    $.ajax({
+        url: "https://www.unrealengine.com/epic-games-header?useServerUrl=true&lang=en-US",
+        context: document.body
+    }).done(function (data) {
+        $("#head").html(data);
+    }).fail(function( jqXHR, textStatus, errorThrown ) {
+	});
+});
+
+window._egNavAfterConfigSet = function() {
+	_epicGamesNav.hideSignIn = true;
+	_epicGamesNav.hideLocale = true;
+};
+</script>
+{% endblock %}
\ No newline at end of file
