UBT new plugin test - generate solution doesn't include create new target

diff --git a/Engine/Plugins/Experimental/StateGraph/Tests/StateGraphTests/Private/StateGraphTestsExample.cpp b/Engine/Plugins/Experimental/StateGraph/Tests/StateGraphTests/Private/StateGraphTestsExample.cpp
new file mode 100644
--- /dev/null
+++ b/Engine/Plugins/Experimental/StateGraph/Tests/StateGraphTests/Private/StateGraphTestsExample.cpp
@@ -0,0 +1,237 @@
+// Copyright Epic Games, Inc. All Rights Reserved.
+
+#include "Containers/Ticker.h"
+#include "Misc/ConfigCacheIni.h"
+#include "StateGraph.h"
+#include "TestHarness.h"
+
+namespace UE::StateGraphTests::Example
+{
+
+// This example shows many of the features provided by the StateGraph module. See class descriptions below for how nodes relate to each other within the state graph.
+// Running this produces the following debug output at the end:
+//
+// [UE::FStateGraph::Run] [Example] Completed (Duration=0.001272)
+// [UE::FStateGraph::LogDebugInfo] [Example] Status=Completed Nodes=11 Duration=0.001272
+// [UE::FStateGraph::LogDebugInfo] [Example.Base.NodeA] Status=Completed Duration=0.000023 Dependencies(None)
+// [UE::FStateGraph::LogDebugInfo] [Example.Base.NodeB] Status=Completed Duration=0.000258 Dependencies(None)
+// [UE::FStateGraph::LogDebugInfo] [Example.Base.NodeC] Status=Completed Duration=0.000206 Dependencies(Completed=Base.NodeA,Base.NodeB)
+// [UE::FStateGraph::LogDebugInfo] [Example.Base.NodeD] Status=Completed Duration=0.000013 Dependencies(Completed=Base.NodeC,Derived.NodeF)
+// [UE::FStateGraph::LogDebugInfo] [Example.Base.NodeE] Status=Completed Duration=0.000190 Dependencies(Completed=Base.NodeC)
+// [UE::FStateGraph::LogDebugInfo] [Example.Derived.NodeF] Status=Completed Duration=0.000556 Dependencies(Completed=Base.NodeC)
+// [UE::FStateGraph::LogDebugInfo] [Example.Derived.NodeG] Status=Completed Duration=0.000056 Dependencies(None)
+// [UE::FStateGraph::LogDebugInfo] [Example.Derived.NodeH] Status=Completed Duration=0.000370 Dependencies(Completed=Derived.NodeG)
+// [UE::FStateGraph::LogDebugInfo] [Example.Module.NodeA] Status=Completed Duration=0.000014 Dependencies(None)
+// [UE::FStateGraph::LogDebugInfo] [Example.Module.NodeB] Status=Completed Duration=0.000321 Dependencies(Completed=Module.NodeA)
+// [UE::FStateGraph::LogDebugInfo] [Example.Module.NodeC] Status=Completed Duration=0.000015 Dependencies(Completed=Module.NodeB)
+//
+// To see the graph work through all nodes, build and run: StateGraphTests.exe [Example] --log --extra-args -LogCmds="LogStateGraph All" -AbsLog=D:\StateGraph.log
+
+
+// Use custom class nodes like this to manage functions and data for a single node, as opposed to the single function nodes seen below.
+// Instances of this are added to the state graph in the classes below.
+class FExampleClassNode : public FStateGraphNode
+{
+public:
+
+	FExampleClassNode(FName InName) :
+		FStateGraphNode(InName)
+	{}
+
+	float Input = 0.1f;
+	float Output = 0;
+
+protected:
+
+	virtual void Start()
+	{
+		// Code to setup async task...
+		FTSTicker::GetCoreTicker().AddTicker(FTickerDelegate::CreateSP(this, &FExampleClassNode::RunAsync), Input);
+	}
+
+	bool RunAsync(float DeltaTime)
+	{
+		if (++RunState < CustomConfig)
+		{
+			return true;
+		}
+
+		FinishAsync();
+		return false;
+	}
+
+	void FinishAsync()
+	{
+		// Code to process result of async task...
+		Output = 0.2f;
+		Complete();
+	}
+
+	virtual void UpdateConfig()
+	{
+		// Override to read custom config for the node.
+		FStateGraphNode::UpdateConfig();
+		GConfig->GetInt(*GetConfigSectionName(), TEXT("CustomConfig"), CustomConfig, GEngineIni);
+	}
+
+	int32 RunState = 0;
+	int32 CustomConfig = 3;
+};
+
+using FExampleClassNodePtr = TSharedPtr<FExampleClassNode, ESPMode::ThreadSafe>;
+
+
+// Helper function to simulate async tasks in simple function nodes below.
+void ExampleAsync(float Delay, const FStateGraphNodeFunctionComplete& Complete)
+{
+	FTSTicker::GetCoreTicker().AddTicker(FTickerDelegate::CreateLambda([Complete](float DeltaTime) {
+		Complete();
+		return false;
+	}), Delay);
+}
+
+
+// FExampleBase - Base class of some function nodes that perform some work.
+class FExampleBase : public TSharedFromThis<FExampleBase>
+{
+public:
+
+	FExampleBase(const FStateGraphRef& InStateGraph) :
+		StateGraph(InStateGraph)
+	{}
+
+	virtual ~FExampleBase() = default;
+
+	virtual void SetupStateGraph()
+	{
+		// CreateNode supports any function type that can be used to create delegates with: lambdas, members, static, raw, etc.
+
+		// Since NodeA and NodeB have no dependencies, they run first concurrently.
+		StateGraph->CreateNode("Base.NodeA", this, [](FStateGraphNodeFunctionComplete Complete) {	Complete(); });
+		StateGraph->CreateNode("Base.NodeB", this, [](FStateGraphNodeFunctionComplete Complete) {	ExampleAsync(0.1f, Complete); });
+
+		// CreateNode also supports custom node types by explicitly specifying the type and passing arguments to the constructor.
+		FExampleClassNodePtr NodeC = StateGraph->CreateNode<FExampleClassNode>("Base.NodeC");
+		// Add dependencies so NodeC isn't run until both NodeA and NodeB complete.
+		NodeC->Dependencies.Append({ "Base.NodeA", "Base.NodeB" });
+		NodeC->Input = 0.1f;
+
+		// Both NodeD and NodeE will run concurrently once NodeC is complete.
+		StateGraph->CreateNode("Base.NodeD", &FExampleBase::NodeD)
+			->Dependencies.Add("Base.NodeC");
+		StateGraph->CreateNode("Base.NodeE", this, &FExampleBase::NodeE)
+			->Dependencies.Add("Base.NodeC");
+	}
+
+protected:
+
+	static void NodeD(FStateGraphNodeFunctionComplete Complete)
+	{
+		Complete();
+	}
+
+	void NodeE(FStateGraphNodeFunctionComplete Complete)
+	{
+		FExampleClassNodePtr NodeC = StateGraph->GetNode<FExampleClassNode>("Base.NodeC");
+		check(NodeC.IsValid());
+		// Use NodeC Output within this node...
+		ExampleAsync(NodeC->Output, Complete);
+	}
+
+	FStateGraphRef StateGraph;
+};
+
+
+// FExampleModule - Class representing a module that shares no nodes with FExampleBase or FExampleDerived, but is used by FExampleDerived to complete some work.
+DECLARE_DELEGATE_OneParam(FExampleModuleComplete, int32 /* Output */);
+class FExampleModule : public TSharedFromThis<FExampleModule>
+{
+public:
+	FExampleModule(const FStateGraphRef& InStateGraph) :
+		StateGraph(InStateGraph)
+	{}
+
+	void Run(const FExampleModuleComplete& ModuleComplete)
+	{
+		// Run nodes in sequence to produce some output that the caller needs, using the same state graph which may be running unrelated nodes.
+		StateGraph->CreateNode("Module.NodeA", this, [](FStateGraphNodeFunctionComplete Complete) { Complete(); })
+			->Next("Module.NodeB", this, &FExampleModule::NodeB)
+			->Next("Module.NodeC", this, [ModuleComplete](FStateGraphNodeFunctionComplete Complete) { Complete(); ModuleComplete.ExecuteIfBound(42); });
+	}
+
+protected:
+
+	void NodeB(FStateGraphNodeFunctionComplete Complete)
+	{
+		ExampleAsync(0.1f, Complete);
+	}
+
+	FStateGraphRef StateGraph;
+};
+
+using FExampleModulePtr = TSharedPtr<FExampleModule, ESPMode::ThreadSafe>;
+
+
+// FExampleDerived - Derived class that updates the function nodes FExampleBase setup.
+class FExampleDerived : public FExampleBase
+{
+public:
+
+	FExampleDerived(const FStateGraphRef& InStateGraph) :
+		FExampleBase(InStateGraph)
+	{}
+
+	virtual void SetupStateGraph() override
+	{
+		FExampleBase::SetupStateGraph();
+
+		// Insert NodeF between NodeC and NodeD
+		StateGraph->CreateNode("Derived.NodeF", this, &FExampleDerived::NodeF)
+			->Dependencies.Add("Base.NodeC");
+
+		FStateGraphNodeRef* NodeDRef = StateGraph->GetNodeRef("Base.NodeD");
+		check(NodeDRef);
+		(*NodeDRef)->Dependencies.Add("Derived.NodeF");
+
+		// Add some additional nodes.
+		StateGraph->CreateNode("Derived.NodeG", this, [](FStateGraphNodeFunctionComplete Complete) { Complete(); })
+			->Next<FExampleClassNode>("Derived.NodeH");
+	}
+
+protected:
+
+	void NodeF(FStateGraphNodeFunctionComplete Complete)
+	{
+		FExampleModulePtr ModulePtr = MakeShared<FExampleModule>(StateGraph);
+		ModulePtr->Run(FExampleModuleComplete::CreateSPLambda(this, [ModulePtr, Complete](int32 Output)
+		{
+			// Use output...
+			Complete();
+		}));
+	}
+};
+
+
+TEST_CASE("Example demonstrating the features of StateGraph", "[Example]")
+{
+	// Engine does this, but it's needed here to ensure GConfig is initialized for the test.
+	FConfigCacheIni::InitializeConfigSystem();
+
+	FStateGraphRef StateGraph(MakeShared<FStateGraph>("Example"));
+	StateGraph->Initialize();
+
+	TSharedRef<FExampleDerived> Example = MakeShared<FExampleDerived>(StateGraph);
+	Example->SetupStateGraph();
+
+	StateGraph->Run();
+
+	// Tick the fake async tasks until complete.
+	while (StateGraph->GetStatus() == FStateGraph::EStatus::Waiting)
+	{
+		FTSTicker::GetCoreTicker().Tick(0.1f);
+	}
+
+	CHECK(StateGraph->GetStatus() == FStateGraph::EStatus::Completed);
+}
+
+} // UE::StateGraphTests
diff --git a/Engine/Plugins/Experimental/StateGraph/Tests/StateGraphTests/Private/StateGraphTestsUnit.cpp b/Engine/Plugins/Experimental/StateGraph/Tests/StateGraphTests/Private/StateGraphTestsUnit.cpp
new file mode 100644
--- /dev/null
+++ b/Engine/Plugins/Experimental/StateGraph/Tests/StateGraphTests/Private/StateGraphTestsUnit.cpp
@@ -0,0 +1,467 @@
+// Copyright Epic Games, Inc. All Rights Reserved.
+
+#include "Misc/ConfigCacheIni.h"
+#include "Misc/CoreDelegates.h"
+#include "StateGraph.h"
+#include "TestHarness.h"
+
+namespace UE::StateGraphTests::Unit
+{
+
+TEST_CASE("FStateGraph Basic Tests", "[FStateGraph]")
+{
+	FStateGraphRef StateGraph(MakeShared<FStateGraph>("Test"));
+	CHECK(StateGraph->GetName() == "Test");
+	CHECK(StateGraph->GetStatus() == FStateGraph::EStatus::NotStarted);
+	CHECK(FString(StateGraph->GetStatusName()) == FString(TEXT("NotStarted")));
+	StateGraph->Run();
+	CHECK(StateGraph->GetStatus() == FStateGraph::EStatus::Completed);
+}
+
+class FTestNode : public FStateGraphNode
+{
+public:
+	FTestNode(FName InName) :
+		FStateGraphNode(InName)
+	{}
+
+	FTestNode(FName InName, const TSet<FName>& InDependencies) :
+		FStateGraphNode(InName)
+	{
+		Dependencies.Append(InDependencies);
+	}
+
+	virtual void Start() override {}
+	virtual void Removed() override { bWasRemovedCalled = true; }
+	virtual void TimedOut() override { bWasTimedOutCalled = true; }
+	virtual void UpdateConfig() override
+	{
+		FStateGraphNode::UpdateConfig();
+		GConfig->GetInt(*GetConfigSectionName(), TEXT("CustomConfig"), CustomConfig, GEngineIni);
+	}
+
+	bool bWasRemovedCalled = false;
+	bool bWasTimedOutCalled = false;
+	int32 CustomConfig = 1;
+};
+
+using FTestNodeRef = TSharedRef<FTestNode, ESPMode::ThreadSafe>;
+using FTestNodePtr = TSharedPtr<FTestNode, ESPMode::ThreadSafe>;
+using FTestNodeWeakPtr = TWeakPtr<FTestNode, ESPMode::ThreadSafe>;
+
+TEST_CASE("FStateGraphNode Basic Tests", "[FStateGraphNode]")
+{
+	FStateGraphRef StateGraph(MakeShared<FStateGraph>("Test"));
+	FStateGraphNodeRef TestNodeA = StateGraph->CreateNode<FTestNode>("TestNodeA");
+	CHECK(TestNodeA->GetName() == "TestNodeA");
+	CHECK(TestNodeA->GetStatus() == FStateGraphNode::EStatus::NotStarted);
+	CHECK(FString(TestNodeA->GetStatusName()) == FString(TEXT("NotStarted")));
+	CHECK(&TestNodeA.Get() == &StateGraph->GetNodeRef("TestNodeA")->Get());
+	CHECK(&TestNodeA.Get() == StateGraph->GetNode<FTestNode>("TestNodeA").Get());
+
+	StateGraph->Run();
+	CHECK(StateGraph->GetStatus() == FStateGraph::EStatus::Waiting);
+	CHECK(TestNodeA->GetStatus() == FStateGraphNode::EStatus::Started);
+
+	TestNodeA->Complete();
+	CHECK(StateGraph->GetStatus() == FStateGraph::EStatus::Completed);
+	CHECK(TestNodeA->GetStatus() == FStateGraphNode::EStatus::Completed);
+}
+
+TEST_CASE("FStateGraphNodeFunction Basic Tests", "[FStateGraphNodeFunction]")
+{
+	FStateGraphRef StateGraph(MakeShared<FStateGraph>("Test"));
+	FStateGraphNodeFunctionComplete Complete;
+	FStateGraphNodeFunctionRef TestNodeA = StateGraph->CreateNode("TestNodeA", [&Complete](FStateGraphNodeFunctionComplete InComplete) { Complete = InComplete; });
+	CHECK(TestNodeA->GetStatus() == FStateGraphNode::EStatus::NotStarted);
+	CHECK(FString(TestNodeA->GetStatusName()) == FString(TEXT("NotStarted")));
+
+	StateGraph->Run();
+	CHECK(StateGraph->GetStatus() == FStateGraph::EStatus::Waiting);
+	CHECK(TestNodeA->GetStatus() == FStateGraphNode::EStatus::Started);
+
+	Complete();
+	CHECK(StateGraph->GetStatus() == FStateGraph::EStatus::Completed);
+	CHECK(TestNodeA->GetStatus() == FStateGraphNode::EStatus::Completed);
+}
+
+class FClassNode : public FStateGraphNode
+{
+public:
+	FClassNode(FName InName) : FStateGraphNode(InName) {}
+	virtual void Start() { Complete(); }
+};
+
+void StaticNode(FStateGraphNodeFunctionComplete Complete) { Complete(); }
+
+class FRawClass
+{
+public:
+	void RawNode(FStateGraphNodeFunctionComplete Complete) { Complete(); }
+};
+
+class FSPClass : public TSharedFromThis<FSPClass>
+{
+public:
+	void SPNode(FStateGraphNodeFunctionComplete Complete) { Complete(); }
+};
+
+TEST_CASE("CreateNode for each supported type", "[CreateNodeFunctions]")
+{
+	FStateGraphRef StateGraph(MakeShared<FStateGraph>("Test"));
+	FRawClass RawClass;
+	TSharedRef<FSPClass> SPClass(MakeShared<FSPClass>());
+
+	StateGraph->CreateNode<FClassNode>("TestClass")
+		->Next("TestLambda", [](FStateGraphNodeFunctionComplete Complete) { Complete(); })
+		->Next("TestStatic", &StaticNode)
+		->Next("TestRaw", &RawClass, &FRawClass::RawNode)
+		->Next("TestSP", &SPClass.Get(), &FSPClass::SPNode)
+		->Next("TestSPLambda", &SPClass.Get(), [](FStateGraphNodeFunctionComplete Complete) { Complete(); });
+
+	StateGraph->Run();
+	CHECK(StateGraph->GetStatus() == FStateGraph::EStatus::Completed);
+}
+
+TEST_CASE("Node dependencies", "[NodeDependencies]")
+{
+	FStateGraphRef StateGraph(MakeShared<FStateGraph>("Test"));
+
+	FStateGraphNodeRef TestNodeA = StateGraph->CreateNode<FTestNode>("TestNodeA");
+	FStateGraphNodeRef TestNodeB = StateGraph->CreateNode<FTestNode>("TestNodeB", TSet<FName>({"TestNodeA"}));
+	FStateGraphNodeRef TestNodeC = StateGraph->CreateNode<FTestNode>("TestNodeC", TSet<FName>({"TestNodeB", "TestNodeD"}));
+
+	StateGraph->Run();
+	CHECK(StateGraph->GetStatus() == FStateGraph::EStatus::Waiting);
+	CHECK(TestNodeA->GetStatus() == FStateGraphNode::EStatus::Started);
+	CHECK(TestNodeB->GetStatus() == FStateGraphNode::EStatus::Blocked);
+	CHECK(TestNodeC->GetStatus() == FStateGraphNode::EStatus::Blocked);
+
+	TestNodeA->Complete();
+	CHECK(StateGraph->GetStatus() == FStateGraph::EStatus::Waiting);
+	CHECK(TestNodeA->GetStatus() == FStateGraphNode::EStatus::Completed);
+	CHECK(TestNodeB->GetStatus() == FStateGraphNode::EStatus::Started);
+	CHECK(TestNodeC->GetStatus() == FStateGraphNode::EStatus::Blocked);
+
+	TestNodeB->Complete();
+	CHECK(StateGraph->GetStatus() == FStateGraph::EStatus::Blocked);
+	CHECK(TestNodeB->GetStatus() == FStateGraphNode::EStatus::Completed);
+	CHECK(TestNodeC->GetStatus() == FStateGraphNode::EStatus::Blocked);
+
+	FStateGraphNodeRef TestNodeD = StateGraph->CreateNode<FTestNode>("TestNodeD");
+	StateGraph->Run();
+	CHECK(StateGraph->GetStatus() == FStateGraph::EStatus::Waiting);
+	CHECK(TestNodeC->GetStatus() == FStateGraphNode::EStatus::Blocked);
+	CHECK(TestNodeD->GetStatus() == FStateGraphNode::EStatus::Started);
+
+	TestNodeD->Complete();
+	CHECK(StateGraph->GetStatus() == FStateGraph::EStatus::Waiting);
+	CHECK(TestNodeC->GetStatus() == FStateGraphNode::EStatus::Started);
+	CHECK(TestNodeD->GetStatus() == FStateGraphNode::EStatus::Completed);
+
+	TestNodeC->Complete();
+	CHECK(StateGraph->GetStatus() == FStateGraph::EStatus::Completed);
+	CHECK(TestNodeC->GetStatus() == FStateGraphNode::EStatus::Completed);
+}
+
+TEST_CASE("Adding and reusing nodes", "[AddingNodes]")
+{
+	FStateGraphRef StateGraphA(MakeShared<FStateGraph>("TestB"));
+	FStateGraphNodeRef TestNodeA = MakeShared<FTestNode>("TestNodeA");
+
+	CHECK(StateGraphA->AddNode(TestNodeA));
+	StateGraphA->Run();
+	CHECK(StateGraphA->GetStatus() == FStateGraph::EStatus::Waiting);
+	CHECK(TestNodeA->GetStatus() == FStateGraphNode::EStatus::Started);
+
+	StateGraphA->RemoveNode(TestNodeA->GetName());
+	StateGraphA->Run();
+	CHECK(StateGraphA->GetStatus() == FStateGraph::EStatus::Completed);
+	CHECK(TestNodeA->GetStatus() == FStateGraphNode::EStatus::Started);
+
+	FStateGraphRef StateGraphB(MakeShared<FStateGraph>("TestB"));
+	CHECK(StateGraphB->AddNode(TestNodeA));
+	StateGraphB->Run();
+	CHECK(StateGraphB->GetStatus() == FStateGraph::EStatus::Waiting);
+	CHECK(TestNodeA->GetStatus() == FStateGraphNode::EStatus::Started);
+
+	TestNodeA->Complete();
+	CHECK(StateGraphB->GetStatus() == FStateGraph::EStatus::Completed);
+	CHECK(TestNodeA->GetStatus() == FStateGraphNode::EStatus::Completed);
+
+	StateGraphA->Reset();
+	CHECK(StateGraphA->GetStatus() == FStateGraph::EStatus::NotStarted);
+	CHECK(!StateGraphA->AddNode(TestNodeA));
+
+	StateGraphB->RemoveNode(TestNodeA->GetName());
+	CHECK(StateGraphA->AddNode(TestNodeA));
+	StateGraphA->Run();
+	CHECK(StateGraphA->GetStatus() == FStateGraph::EStatus::Completed);
+	CHECK(TestNodeA->GetStatus() == FStateGraphNode::EStatus::Completed);
+
+	FStateGraphNodeRef TestNodeA2 = MakeShared<FTestNode>("TestNodeA");
+	CHECK(!StateGraphA->AddNode(TestNodeA2));
+}
+
+TEST_CASE("Removing nodes", "[RemoveNodes]")
+{
+	FStateGraphRef StateGraph(MakeShared<FStateGraph>("Test"));
+
+	FTestNodePtr TestNodeA = StateGraph->CreateNode<FTestNode>("TestNodeA");
+	FTestNodeWeakPtr WeakTestNodeA = TestNodeA;
+	StateGraph->RemoveNode("TestNodeA");
+	CHECK(TestNodeA->GetStatus() == FStateGraphNode::EStatus::NotStarted);
+	CHECK(TestNodeA->bWasRemovedCalled);
+	TestNodeA.Reset();
+	CHECK(!WeakTestNodeA.IsValid());
+
+	WeakTestNodeA = StateGraph->CreateNode<FTestNode>("TestNodeA");
+	StateGraph->RemoveNode("TestNodeA");
+	CHECK(!WeakTestNodeA.IsValid());
+
+	WeakTestNodeA = StateGraph->CreateNode<FTestNode>("TestNodeA");
+	FStateGraphNodePtr TestNodeB = StateGraph->CreateNode("TestNodeB", [StateGraph](FStateGraphNodeFunctionComplete Complete)
+	{
+		StateGraph->RemoveNode("TestNodeA");
+		StateGraph->RemoveNode("TestNodeC");
+	});
+	FStateGraphNodeWeakPtr WeakTestNodeC = StateGraph->CreateNode<FTestNode>("TestNodeC");
+	StateGraph->Run();
+	CHECK(!WeakTestNodeA.IsValid());
+	CHECK(!WeakTestNodeC.IsValid());
+	StateGraph->RemoveNode("TestNodeB");
+
+	StateGraph->CreateNode<FTestNode>("TestNodeA")
+		->Next<FTestNode>("TestNodeB")
+		->Next<FTestNode>("TestNodeC");
+	StateGraph->RemoveAllNodes();
+	CHECK(StateGraph->GetNodeRef("TestNodeA") == nullptr);
+	CHECK(StateGraph->GetNodeRef("TestNodeB") == nullptr);
+	CHECK(StateGraph->GetNodeRef("TestNodeC") == nullptr);
+}
+
+TEST_CASE("Resetting nodes", "[ResetNodes]")
+{
+	FStateGraphRef StateGraph(MakeShared<FStateGraph>("Test"));
+
+	int32 StartCount = 0;
+	FStateGraphNodeRef TestNodeA = StateGraph->CreateNode("TestNodeA", [&StartCount, &TestNodeA](FStateGraphNodeFunctionComplete Complete)
+	{
+		CHECK(TestNodeA->GetStatus() == FStateGraphNode::EStatus::Started);
+
+		if (StartCount == 0)
+		{
+			TestNodeA->Reset();
+		}
+		else
+		{
+			Complete();
+		}
+
+		++StartCount;
+	});
+
+	StateGraph->Run();
+	CHECK(StartCount == 1);
+	CHECK(TestNodeA->GetStatus() == FStateGraphNode::EStatus::NotStarted);
+
+	StateGraph->Run();
+	CHECK(StartCount == 2);
+	CHECK(TestNodeA->GetStatus() == FStateGraphNode::EStatus::Completed);
+
+	TestNodeA->Reset();
+	StateGraph->Run();
+	CHECK(StartCount == 3);
+}
+
+TEST_CASE("Resetting state graph", "[ResetStateGraph]")
+{
+	FStateGraphRef StateGraph(MakeShared<FStateGraph>("Test"));
+
+	int32 StartCount = 0;
+	FStateGraphNodeRef TestNodeA = StateGraph->CreateNode("TestNodeA", [&StartCount, &StateGraph](FStateGraphNodeFunctionComplete Complete)
+	{
+		if (StartCount == 0)
+		{
+			StateGraph->Reset();
+		}
+		else
+		{
+			Complete();
+		}
+
+		++StartCount;
+	});
+
+	StateGraph->Run();
+	CHECK(StartCount == 1);
+	CHECK(TestNodeA->GetStatus() == FStateGraphNode::EStatus::NotStarted);
+	CHECK(StateGraph->GetStatus() == FStateGraph::EStatus::NotStarted);
+
+	StateGraph->Run();
+	CHECK(StartCount == 2);
+	CHECK(TestNodeA->GetStatus() == FStateGraphNode::EStatus::Completed);
+	CHECK(StateGraph->GetStatus() == FStateGraph::EStatus::Completed);
+}
+
+TEST_CASE("Deleting state graph", "[DeleteStateGraph]")
+{
+	FStateGraphPtr StateGraph(MakeShared<FStateGraph>("Test"));
+	FStateGraphWeakPtr WeakStateGraph(StateGraph);
+	FStateGraphNodeRef TestNodeA = StateGraph->CreateNode("TestNodeA", [&StateGraph, &WeakStateGraph](FStateGraphNodeFunctionComplete Complete)
+	{
+		StateGraph.Reset();
+		CHECK(!WeakStateGraph.IsValid());
+	});
+
+	StateGraph->Run();
+	CHECK(!WeakStateGraph.IsValid());
+	CHECK(TestNodeA->GetStateGraphName() == NAME_None);
+}
+
+TEST_CASE("Pausing state graph", "[PauseStateGraph]")
+{
+	FStateGraphRef StateGraph(MakeShared<FStateGraph>("Test"));
+	FStateGraphNodeRef TestNodeA = StateGraph->CreateNode("TestNodeA", [&StateGraph](FStateGraphNodeFunctionComplete Complete) { StateGraph->Pause(); });
+	FStateGraphNodeRef TestNodeB = StateGraph->CreateNode<FTestNode>("TestNodeB", TSet<FName>({"TestNodeA"}));
+	FStateGraphNodeRef TestNodeC = StateGraph->CreateNode("TestNodeC", [](FStateGraphNodeFunctionComplete Complete) { Complete(); });
+	TestNodeC->Dependencies.Add("TestNodeB");
+
+	StateGraph->Run();
+	CHECK(StateGraph->GetStatus() == FStateGraph::EStatus::Paused);
+	CHECK(TestNodeA->GetStatus() == FStateGraphNode::EStatus::Started);
+
+	TestNodeA->Complete();
+	CHECK(StateGraph->GetStatus() == FStateGraph::EStatus::Paused);
+	CHECK(TestNodeA->GetStatus() == FStateGraphNode::EStatus::Completed);
+	CHECK(TestNodeB->GetStatus() != FStateGraphNode::EStatus::Started);
+
+	StateGraph->Run();
+	CHECK(StateGraph->GetStatus() == FStateGraph::EStatus::Waiting);
+	CHECK(TestNodeA->GetStatus() == FStateGraphNode::EStatus::Completed);
+	CHECK(TestNodeB->GetStatus() == FStateGraphNode::EStatus::Started);
+	CHECK(TestNodeC->GetStatus() == FStateGraphNode::EStatus::Blocked);
+
+	StateGraph->Pause();
+	TestNodeB->Complete();
+	CHECK(StateGraph->GetStatus() == FStateGraph::EStatus::Paused);
+	CHECK(TestNodeB->GetStatus() == FStateGraphNode::EStatus::Completed);
+	CHECK(TestNodeC->GetStatus() == FStateGraphNode::EStatus::Blocked);
+
+	StateGraph->Run();
+	CHECK(StateGraph->GetStatus() == FStateGraph::EStatus::Completed);
+	CHECK(TestNodeC->GetStatus() == FStateGraphNode::EStatus::Completed);
+}
+
+TEST_CASE("State graph timeout", "[StateGraphTimeout]")
+{
+	FStateGraphRef StateGraph(MakeShared<FStateGraph>("Test"));
+	StateGraph->SetTimeout(0.01);
+	bool bOnTimedOutCalled = false;
+	StateGraph->OnStatusChanged.AddLambda([&bOnTimedOutCalled](FStateGraph&, FStateGraph::EStatus OldStatus, FStateGraph::EStatus NewStatus)
+	{
+		if (NewStatus == FStateGraph::EStatus::TimedOut)
+		{
+			bOnTimedOutCalled = true;
+		}
+	});
+
+	FStateGraphNodeRef TestNodeA = StateGraph->CreateNode("TestNodeA", [](FStateGraphNodeFunctionComplete Complete) {});
+	StateGraph->Run();
+
+	while (StateGraph->GetStatus() == FStateGraph::EStatus::Waiting)
+	{
+		FTSTicker::GetCoreTicker().Tick(0.1f);
+	}
+
+	CHECK(StateGraph->GetStatus() == FStateGraph::EStatus::TimedOut);
+	CHECK(bOnTimedOutCalled);
+}
+
+TEST_CASE("Node timeout", "[NodeTimeout]")
+{
+	FStateGraphRef StateGraph(MakeShared<FStateGraph>("Test"));
+	bool bOnNodeTimedOutCalled = false;
+	StateGraph->OnNodeStatusChanged.AddLambda([&bOnNodeTimedOutCalled](FStateGraphNode&, FStateGraphNode::EStatus OldStatus, FStateGraphNode::EStatus NewStatus)
+	{
+		if (NewStatus == FStateGraphNode::EStatus::TimedOut)
+		{
+			bOnNodeTimedOutCalled = true;
+		}
+	});
+
+	FTestNodeRef TestNodeA = StateGraph->CreateNode<FTestNode>("TestNodeA");
+	TestNodeA->SetTimeout(0.01);
+	StateGraph->Run();
+
+	while (StateGraph->GetStatus() == FStateGraph::EStatus::Waiting)
+	{
+		FTSTicker::GetCoreTicker().Tick(0.1f);
+	}
+
+	CHECK(TestNodeA->GetStatus() == FStateGraphNode::EStatus::TimedOut);
+	CHECK(StateGraph->GetStatus() == FStateGraph::EStatus::Blocked);
+	CHECK(TestNodeA->bWasTimedOutCalled);
+	CHECK(bOnNodeTimedOutCalled);
+
+	TestNodeA->Complete();
+	CHECK(StateGraph->GetStatus() == FStateGraph::EStatus::Completed);
+}
+
+TEST_CASE("State graph config", "[StateGraphConfig]")
+{
+	FConfigCacheIni::InitializeConfigSystem();
+	FConfigFile* EngineIni = GConfig->FindConfigFile(GEngineIni);
+	check(EngineIni);
+	EngineIni->CombineFromBuffer(TEXT("[StateGraph.Test]\nTimeout=0.01\n"), GEngineIni);
+
+	FStateGraphRef StateGraph(MakeShared<FStateGraph>("Test"));
+	StateGraph->Initialize();
+	FStateGraphNodeRef TestNodeA = StateGraph->CreateNode("TestNodeA", [](FStateGraphNodeFunctionComplete Complete) {});
+	StateGraph->Run();
+
+	while (StateGraph->GetStatus() == FStateGraph::EStatus::Waiting)
+	{
+		FTSTicker::GetCoreTicker().Tick(0.1f);
+	}
+
+	CHECK(StateGraph->GetStatus() == FStateGraph::EStatus::TimedOut);
+	TestNodeA->Complete();
+	CHECK(StateGraph->GetStatus() == FStateGraph::EStatus::TimedOut);
+
+	EngineIni->CombineFromBuffer(TEXT("[StateGraph.Test]\nTimeout=0\n"), GEngineIni);
+	FCoreDelegates::TSOnConfigSectionsChanged().Broadcast(GEngineIni, TSet<FString>({ TEXT("StateGraph.Test") }));
+
+	StateGraph->Run();
+	CHECK(StateGraph->GetStatus() == FStateGraph::EStatus::Completed);
+}
+
+TEST_CASE("Node config", "[NodeConfig]")
+{
+	FConfigCacheIni::InitializeConfigSystem();
+	FConfigFile* EngineIni = GConfig->FindConfigFile(GEngineIni);
+	check(EngineIni);
+	EngineIni->CombineFromBuffer(TEXT("[StateGraph.Test.TestNodeA]\nTimeout=0.01\nCustomConfig=2\n"), GEngineIni);
+
+	FStateGraphRef StateGraph(MakeShared<FStateGraph>("Test"));
+	StateGraph->Initialize();
+	FTestNodeRef TestNodeA = StateGraph->CreateNode<FTestNode>("TestNodeA");
+	StateGraph->Run();
+
+	while (StateGraph->GetStatus() == FStateGraph::EStatus::Waiting)
+	{
+		FTSTicker::GetCoreTicker().Tick(0.1f);
+	}
+
+	CHECK(TestNodeA->GetStatus() == FStateGraphNode::EStatus::TimedOut);
+	CHECK(StateGraph->GetStatus() == FStateGraph::EStatus::Blocked);
+	CHECK(TestNodeA->bWasTimedOutCalled);
+	CHECK(TestNodeA->CustomConfig == 2);
+
+	EngineIni->CombineFromBuffer(TEXT("[StateGraph.Test.TestNodeA]\nTimeout=0.01\nCustomConfig=3\n"), GEngineIni);
+	FCoreDelegates::TSOnConfigSectionsChanged().Broadcast(GEngineIni, TSet<FString>({ TEXT("StateGraph.Test.TestNodeA") }));
+
+	CHECK(TestNodeA->CustomConfig == 3);
+}
+
+} // UE::StateGraphTests
diff --git a/Engine/Plugins/Experimental/StateGraph/Tests/StateGraphTests/StateGraphTests.Build.cs b/Engine/Plugins/Experimental/StateGraph/Tests/StateGraphTests/StateGraphTests.Build.cs
new file mode 100644
--- /dev/null
+++ b/Engine/Plugins/Experimental/StateGraph/Tests/StateGraphTests/StateGraphTests.Build.cs
@@ -0,0 +1,19 @@
+// Copyright Epic Games, Inc. All Rights Reserved.
+
+using UnrealBuildTool;
+
+public class StateGraphTests : TestModuleRules
+{
+	public StateGraphTests(ReadOnlyTargetRules Target) : base(Target)
+	{
+		OptimizeCode = CodeOptimization.Never;
+
+		PrivateDependencyModuleNames.AddRange(
+			new string[] {
+				"Core",
+				"StateGraph"
+			});
+
+		UpdateBuildGraphPropertiesFile(new Metadata() { TestName = "StateGraph", TestShortName = "State Graph"});
+	}
+}
diff --git a/Engine/Plugins/Experimental/StateGraph/Tests/StateGraphTests/StateGraphTests.Target.cs b/Engine/Plugins/Experimental/StateGraph/Tests/StateGraphTests/StateGraphTests.Target.cs
new file mode 100644
--- /dev/null
+++ b/Engine/Plugins/Experimental/StateGraph/Tests/StateGraphTests/StateGraphTests.Target.cs
@@ -0,0 +1,12 @@
+// Copyright Epic Games, Inc. All Rights Reserved.
+
+using UnrealBuildTool;
+
+public class StateGraphTestsTarget : TestTargetRules
+{
+	public StateGraphTestsTarget(TargetInfo Target) : base(Target)
+	{
+		bForceEnableExceptions = true;
+		EnablePlugins.Add("StateGraph");
+	}
+}
diff --git a/Engine/Source/Programs/StateGraphTests/Source/StateGraphTestsExample.cpp b/Engine/Source/Programs/StateGraphTests/Source/StateGraphTestsExample.cpp
deleted file mode 100644
diff --git a/Engine/Source/Programs/StateGraphTests/Source/StateGraphTestsUnit.cpp b/Engine/Source/Programs/StateGraphTests/Source/StateGraphTestsUnit.cpp
deleted file mode 100644
diff --git a/Engine/Source/Programs/StateGraphTests/StateGraphTests.Build.cs b/Engine/Source/Programs/StateGraphTests/StateGraphTests.Build.cs
deleted file mode 100644
diff --git a/Engine/Source/Programs/StateGraphTests/StateGraphTests.Target.cs b/Engine/Source/Programs/StateGraphTests/StateGraphTests.Target.cs
deleted file mode 100644
diff --git a/Engine/Source/Programs/UnrealBuildTool/Configuration/TargetDescriptor.cs b/Engine/Source/Programs/UnrealBuildTool/Configuration/TargetDescriptor.cs
--- a/Engine/Source/Programs/UnrealBuildTool/Configuration/TargetDescriptor.cs
+++ b/Engine/Source/Programs/UnrealBuildTool/Configuration/TargetDescriptor.cs
@@ -26,7 +26,7 @@ namespace UnrealBuildTool
 		public CommandLineArguments AdditionalArguments;
 		public bool IsTestsTarget = false;
 
-		public static string GetTestedTargetName(string Name)
+		public static string GetTestedName(string Name)
 		{
 			if (Name.EndsWith(TEST_TARGETS_SUFFIX))
 			{
diff --git a/Engine/Source/Programs/UnrealBuildTool/System/ExternalExecution.cs b/Engine/Source/Programs/UnrealBuildTool/System/ExternalExecution.cs
--- a/Engine/Source/Programs/UnrealBuildTool/System/ExternalExecution.cs
+++ b/Engine/Source/Programs/UnrealBuildTool/System/ExternalExecution.cs
@@ -334,7 +334,8 @@ namespace UnrealBuildTool
 				}
 				if (RulesObject.Plugin != null)
 				{
-					ModuleDescriptor? Module = RulesObject.Plugin.Descriptor.Modules?.FirstOrDefault(x => x.Name == RulesObject.Name);
+					string PluginName = !RulesObject.IsTestModule ? RulesObject.Name : TargetDescriptor.GetTestedName(RulesObject.Name);
+					ModuleDescriptor? Module = RulesObject.Plugin.Descriptor.Modules?.FirstOrDefault(x => x.Name == PluginName);
 					if (Module != null)
 					{
 						return GetEngineModuleTypeFromDescriptor(Module);
diff --git a/Engine/Source/Programs/UnrealBuildTool/System/RulesAssembly.cs b/Engine/Source/Programs/UnrealBuildTool/System/RulesAssembly.cs
--- a/Engine/Source/Programs/UnrealBuildTool/System/RulesAssembly.cs
+++ b/Engine/Source/Programs/UnrealBuildTool/System/RulesAssembly.cs
@@ -412,7 +412,7 @@ namespace UnrealBuildTool
 		{
 			if (Target.IsTestTarget && !Target.ExplicitTestsTarget)
 			{
-				ModuleName = TargetDescriptor.GetTestedTargetName(ModuleName);
+				ModuleName = TargetDescriptor.GetTestedName(ModuleName);
 			}
 
 			// Currently, we expect the user's rules object type name to be the same as the module name
@@ -546,7 +546,7 @@ namespace UnrealBuildTool
 				{
 					if (!Target.ExplicitTestsTarget)
 					{
-						if (Target.LaunchModuleName != null && ModuleName == TargetDescriptor.GetTestedTargetName(Target.LaunchModuleName))
+						if (Target.LaunchModuleName != null && ModuleName == TargetDescriptor.GetTestedName(Target.LaunchModuleName))
 						{
 							RulesObject = new TestModuleRules(RulesObject);
 						}
@@ -777,7 +777,7 @@ namespace UnrealBuildTool
 		{
 			if (IsTestTarget)
 			{
-				TargetName = TargetDescriptor.GetTestedTargetName(TargetName);
+				TargetName = TargetDescriptor.GetTestedName(TargetName);
 			}
 
 			if (Architectures == null)
diff --git a/Engine/Source/Programs/UnrealBuildTool/System/RulesCompiler.cs b/Engine/Source/Programs/UnrealBuildTool/System/RulesCompiler.cs
--- a/Engine/Source/Programs/UnrealBuildTool/System/RulesCompiler.cs
+++ b/Engine/Source/Programs/UnrealBuildTool/System/RulesCompiler.cs
@@ -131,6 +131,7 @@ namespace UnrealBuildTool
 			// Find the shared modules, excluding the programs directory. These are used to create an assembly with the bContainsEngineModules flag set to true.
 			Dictionary<FileReference, ModuleRulesContext> ModuleFileToContext = new Dictionary<FileReference, ModuleRulesContext>();
 			ModuleRulesContext DefaultModuleContext = new ModuleRulesContext(Scope, RootDirectories[0]);
+			List<FileReference> ProgramTargetFiles = new List<FileReference>();
 
 			foreach (DirectoryReference RootDirectory in RootDirectories)
 			{
@@ -151,6 +152,8 @@ namespace UnrealBuildTool
 			{
 				ModuleRulesContext PluginsModuleContext = new ModuleRulesContext(PluginsScope, RootDirectories[0]);
 				FindModuleRulesForPlugins(Plugins, PluginsModuleContext, ModuleFileToContext);
+				// Plugin test target rules only
+				FindTestRulesForPlugins(Plugins, PluginsModuleContext, ModuleFileToContext, ProgramTargetFiles);
 			}
 
 			// Create the assembly
@@ -158,7 +161,6 @@ namespace UnrealBuildTool
 			FileReference EngineAssemblyFileName = FileReference.Combine(AssemblyDir, "Intermediate", "Build", "BuildRules", AssemblyPrefix + "Rules" + FrameworkAssemblyExtension);
 			RulesAssembly EngineAssembly = new RulesAssembly(Scope, RootDirectories, Plugins, ModuleFileToContext, new List<FileReference>(), EngineAssemblyFileName, bContainsEngineModules: true, DefaultBuildSettings: DefaultEngineBuildSettingsVersion, bReadOnly: bReadOnly, bSkipCompile: bSkipCompile, bForceCompile: bForceCompile, Parent: Parent, Logger: Logger);
 
-			List<FileReference> ProgramTargetFiles = new List<FileReference>();
 			Dictionary<FileReference, ModuleRulesContext> ProgramModuleFiles = new Dictionary<FileReference, ModuleRulesContext>();
 			foreach (DirectoryReference RootDirectory in RootDirectories)
 			{
