Add deadlock test to StaeGraph.
#rnx

diff --git a/Engine/Plugins/Experimental/StateGraph/Tests/StateGraphTests/Private/StateGraphTestsUnit.cpp b/Engine/Plugins/Experimental/StateGraph/Tests/StateGraphTests/Private/StateGraphTestsUnit.cpp
--- a/Engine/Plugins/Experimental/StateGraph/Tests/StateGraphTests/Private/StateGraphTestsUnit.cpp
+++ b/Engine/Plugins/Experimental/StateGraph/Tests/StateGraphTests/Private/StateGraphTestsUnit.cpp
@@ -164,6 +164,19 @@ TEST_CASE("Node dependencies", "[NodeDependencies]")
 	CHECK(TestNodeC->GetStatus() == FStateGraphNode::EStatus::Completed);
 }
 
+TEST_CASE("Blocked graph", "[BlockedGraph]")
+{
+	FStateGraphRef StateGraph(MakeShared<FStateGraph>("Test"));
+
+	FStateGraphNodeRef TestNodeA = StateGraph->CreateNode<FTestNode>("TestNodeA", TSet<FName>({ "TestNodeB" }));
+	FStateGraphNodeRef TestNodeB = StateGraph->CreateNode<FTestNode>("TestNodeB", TSet<FName>({ "TestNodeA" }));
+
+	StateGraph->Run();
+	CHECK(StateGraph->GetStatus() == FStateGraph::EStatus::Blocked);
+	CHECK(TestNodeA->GetStatus() == FStateGraphNode::EStatus::Blocked);
+	CHECK(TestNodeB->GetStatus() == FStateGraphNode::EStatus::Blocked);
+}
+
 TEST_CASE("Adding and reusing nodes", "[AddingNodes]")
 {
 	FStateGraphRef StateGraphA(MakeShared<FStateGraph>("TestB"));
