gameplay graph tests and some fixes stemming from these tests

- UGraph::GetCompleteNodeHandle will return an empty handle if the index is invalid but has a valid element (or in any situtation where the index does not exist in the graph)
- Made UGraph::CreateEdge private to force usage of UGraph::CreateBulkEdges
- UGraph::CreateBulkEdges can specifiy output edges as well as edge unique index using the FEdgeCreationParameters struct
- UGraphEdge::GetOtherNode will only return a valid node handle if the input node is one of the two nodes in the edge
- UGraph::CreateVertex can create a vertex with a random valid unique index if not otherwise given one
- Remove extra UGraph::VertexEdges book-keeping
- Add UGraphVertex::NumEdges and UGraphVertex::GetEdges
- UGraphHandle::GetElement is removed - should use the typed functions on the vertex/edge/island handles instead
- Made UGraph::GetVertices, UGraph::GetEdges, UGraph::GetIslands, and UGraph::GetProperties public
- FGraphHandle now does not cache a pointer to the element but a pointer to the graph to ensure that once the vertex/edge/island is removed from the graph, every handle becomes incomplete. This also ensures that UGraph::CreateVertex can't create a duplicate vertex handle that points to a different underlying vertex element.
- Remove bHandleIslands from the public UGraph::RemoveEdge to remove a way for users to put the graph in a weird state
- Remove UGraphIsland::bPendingDestroy as it wasn't used
- Fixed an issue where disabling EGraphIslandOperations::Merge on a single island would not end up actually preventing the merge
- Creating an edge with UGraph::CreateEdge will always initialize the a node's island if it didn't exist before.
- UGraph::LoadFromSerializedGraph will now merge/split islands depending on the edge connectivity

#tests low level tests

diff --git a/Engine/Plugins/Experimental/GameplayGraph/Source/GameplayGraph/Private/Graph/Graph.cpp b/Engine/Plugins/Experimental/GameplayGraph/Source/GameplayGraph/Private/Graph/Graph.cpp
--- a/Engine/Plugins/Experimental/GameplayGraph/Source/GameplayGraph/Private/Graph/Graph.cpp
+++ b/Engine/Plugins/Experimental/GameplayGraph/Source/GameplayGraph/Private/Graph/Graph.cpp
@@ -5,9 +5,17 @@
 #include "Graph/GraphVertex.h"
 #include "Graph/Algorithms/Connectivity/ConnectedComponents.h"
 #include "GenericPlatform/GenericPlatformMath.h"
+#include "HAL/IConsoleManager.h"
 
 #include UE_INLINE_GENERATED_CPP_BY_NAME(Graph)
 
+
+namespace Graph
+{
+	static bool bFixupGraphOnLoad = true;
+	static FAutoConsoleVariableRef FixupGraphOnLoadCVar(TEXT("GameplayGraph.FixupGraphOnLoad"), bFixupGraphOnLoad, TEXT("Merge and split islands when loaded from serialization."));
+}
+
 FSerializableGraph UGraph::GetSerializableGraph() const
 {
 	FSerializableGraph Out;
@@ -62,6 +70,8 @@ void UGraph::LoadFromSerializedGraph(const FSerializableGraph& Input)
 		CreateVertex(Handle.GetUniqueIndex());
 	}
 
+	TArray<FGraphEdgeHandle> AllEdges;
+	AllEdges.Reserve(Input.Edges.Num());
 	for (const TPair<FGraphEdgeHandle, FSerializedEdgeData>& Kvp : Input.Edges)
 	{
 		if (!Kvp.Key.IsValid())
@@ -69,14 +79,16 @@ void UGraph::LoadFromSerializedGraph(const FSerializableGraph& Input)
 			continue;
 		}
 
-		CreateEdge(
+		AllEdges.Emplace(CreateEdge(
 			GetCompleteNodeHandle(Kvp.Value.Node1),
 			GetCompleteNodeHandle(Kvp.Value.Node2),
 			Kvp.Key.GetUniqueIndex(),
 			false
-		);
+		));
 	}
 
+	TArray<FGraphIslandHandle> AllIslands;
+	AllIslands.Reserve(Input.Islands.Num());
 	for (const TPair<FGraphIslandHandle, FSerializedIslandData>& Kvp : Input.Islands)
 	{
 		if (!Kvp.Key.IsValid())
@@ -96,7 +108,17 @@ void UGraph::LoadFromSerializedGraph(const FSerializableGraph& Input)
 			}
 		}
 
-		CreateIsland(IslandVertices, Kvp.Key.GetUniqueIndex());
+		AllIslands.Emplace(CreateIsland(IslandVertices, Kvp.Key.GetUniqueIndex()));
+	}
+
+	if (Graph::bFixupGraphOnLoad)
+	{
+		MergeOrCreateIslands(AllEdges);
+
+		for (const FGraphIslandHandle& NewIslandHandle : AllIslands)
+		{
+			RemoveOrSplitIsland(NewIslandHandle.GetIsland());
+		}
 	}
 }
 
@@ -124,10 +146,17 @@ FGraphVertexHandle UGraph::CreateVertex(FGraphUniqueIndex InUniqueIndex)
 
 	if (!ensure(InUniqueIndex.IsValid()))
 	{
-		return{};
+		return {};
 	}
-	Vertex->OnCreate();
+
 	Vertex->SetUniqueIndex(InUniqueIndex);
+	Vertex->SetParentGraph(this);
+	if (Vertices.Contains(Vertex->Handle()))
+	{
+		return {};
+	}
+
+	Vertex->OnCreate();
 	RegisterVertex(Vertex);
 	OnVertexCreated.Broadcast(Vertex->Handle());
 	return Vertex->Handle();
@@ -141,12 +170,12 @@ void UGraph::RegisterVertex(TObjectPtr<UGraphVertex> Vertex)
 		return;
 	}
 
-	const FGraphVertexHandle Handle = Vertex->Handle();
 	Vertex->SetParentGraph(this);
+	const FGraphVertexHandle Handle = Vertex->Handle();
 	Vertices.Add(Handle, Vertex);
 }
 
-FGraphEdgeHandle UGraph::CreateEdge(FGraphVertexHandle Node1, FGraphVertexHandle Node2, FGraphUniqueIndex InUniqueIndex, bool bAddToIslands)
+FGraphEdgeHandle UGraph::CreateEdge(FGraphVertexHandle Node1, FGraphVertexHandle Node2, FGraphUniqueIndex InUniqueIndex, bool bMergeIslands)
 {
 	TRACE_CPUPROFILER_EVENT_SCOPE(UGraph::CreateEdge);
 	if (!Node1.IsComplete() || !Node2.IsComplete())
@@ -195,7 +224,7 @@ FGraphEdgeHandle UGraph::CreateEdge(FGraphVertexHandle Node1, FGraphVertexHandle
 		{
 			// Possible merge scenario.
 			const bool bIsMerge = Island1 != Island2;
-			if (bIsMerge && !Island1Ptr->IsOperationAllowed(EGraphIslandOperations::Merge) && !Island2Ptr->IsOperationAllowed(EGraphIslandOperations::Merge))
+			if (bIsMerge && (!Island1Ptr->IsOperationAllowed(EGraphIslandOperations::Merge) || !Island2Ptr->IsOperationAllowed(EGraphIslandOperations::Merge)))
 			{
 				return {};
 			}
@@ -223,23 +252,44 @@ FGraphEdgeHandle UGraph::CreateEdge(FGraphVertexHandle Node1, FGraphVertexHandle
 	OnEdgeCreated.Broadcast(Edge->Handle());
 
 	// If we want to keep track of islands, this is where we need to create/merge islands.
-	if (Properties.bGenerateIslands && bAddToIslands)
+	if (Properties.bGenerateIslands)
 	{
-		MergeOrCreateIslands({ Edge->Handle() });
+		// After we've created the edge - we want to see one of the two islands is invalid in which
+		// case we have to initialize that node's island. This is needed or else we might not detect a merge
+		// properly when creating multiple edges at a time in UGraph::CreateBulkEdges.
+		const FGraphIslandHandle& IslandHandle1 = Node1Ptr->GetParentIsland();
+		const FGraphIslandHandle& IslandHandle2 = Node2Ptr->GetParentIsland();
+
+		if (!IslandHandle1.IsValid() || !IslandHandle2.IsValid())
+		{
+			if (UGraphIsland* Island1 = IslandHandle1.GetIsland())
+			{
+				Island1->AddVertex(Node2);
+			}
+			else if (UGraphIsland* Island2 = IslandHandle1.GetIsland())
+			{
+				Island2->AddVertex(Node1);
+			}
+		}
+
+		if (bMergeIslands)
+		{
+			MergeOrCreateIslands({ Edge->Handle() });
+		}
 	}
 
 	return Edge->Handle();
 }
 
-void UGraph::CreateBulkEdges(TArray<TPair<FGraphVertexHandle, FGraphVertexHandle>>&& NodesToConnect)
+void UGraph::CreateBulkEdges(TArray<FEdgeCreationParameters>&& NodesToConnect, TArray<FGraphEdgeHandle>* OutEdges)
 {
 	// Create all edges normally but don't call MergeOrCreateIslands yet. We'll use the bulk function instead.
 	TArray<FGraphEdgeHandle> NewEdges;
 	NewEdges.Reserve(NodesToConnect.Num());
 
-	for (const TPair<FGraphVertexHandle, FGraphVertexHandle>& NodePair : NodesToConnect)
+	for (const FEdgeCreationParameters& Params : NodesToConnect)
 	{
-		if (FGraphEdgeHandle Edge = CreateEdge(NodePair.Key, NodePair.Value, FGraphUniqueIndex::CreateUniqueIndex(), false); Edge.IsValid())
+		if (FGraphEdgeHandle Edge = CreateEdge(Params.VertexHandle1, Params.VertexHandle2, Params.EdgeIndex, false); Edge.IsValid())
 		{
 			NewEdges.Add(Edge);
 		}
@@ -247,7 +297,12 @@ void UGraph::CreateBulkEdges(TArray<TPair<FGraphVertexHandle, FGraphVertexHandle
 
 	// MergeOrCreateIslands incrementally determines island connectivity one edge at a time. This is efficient if we're handling a single edge but less
 	// efficient if we're trying to add a bunch of edges all at the same time since it'll cause a vertex to jump between islands. The Bulk function helps prevent that.
-	MergeOrCreateIslands(MoveTemp(NewEdges));
+	MergeOrCreateIslands(NewEdges);
+
+	if (OutEdges)
+	{
+		*OutEdges = MoveTemp(NewEdges);
+	}
 }
 
 void UGraph::RegisterEdge(TObjectPtr<UGraphEdge> Edge)
@@ -258,12 +313,9 @@ void UGraph::RegisterEdge(TObjectPtr<UGraphEdge> Edge)
 		return;
 	}
 
+	Edge->SetParentGraph(this);
 	const FGraphEdgeHandle Handle = Edge->Handle();
 	Edges.Add(Handle, Edge);
-	Edge->SetParentGraph(this);
-
-	VertexEdges.FindOrAdd(Edge->NodeA()).Add(Handle);
-	VertexEdges.FindOrAdd(Edge->NodeB()).Add(Handle);
 }
 
 FGraphIslandHandle UGraph::CreateIsland(TArray<FGraphVertexHandle> InputNodes, FGraphUniqueIndex InUniqueIndex)
@@ -279,8 +331,10 @@ FGraphIslandHandle UGraph::CreateIsland(TArray<FGraphVertexHandle> InputNodes, F
 	{
 		return {};
 	}
+	Island->SetParentGraph(this);
 	Island->OnCreate();
 	Island->SetUniqueIndex((!InUniqueIndex.IsValid()) ? FGraphUniqueIndex::CreateUniqueIndex() : InUniqueIndex);
+	RegisterIsland(Island);
 
 	OnIslandCreated.Broadcast(Island->Handle());
 	for (const FGraphVertexHandle& Node : InputNodes)
@@ -288,7 +342,6 @@ FGraphIslandHandle UGraph::CreateIsland(TArray<FGraphVertexHandle> InputNodes, F
 		Island->AddVertex(Node);
 	}
 
-	RegisterIsland(Island);
 	return Island->Handle();
 }
 
@@ -323,7 +376,7 @@ void UGraph::RemoveIsland(const FGraphIslandHandle& IslandHandle)
 	Islands.Remove(IslandHandle);
 }
 
-void UGraph::MergeOrCreateIslands(TArray<FGraphEdgeHandle>&& InEdges)
+void UGraph::MergeOrCreateIslands(const TArray<FGraphEdgeHandle>& InEdges)
 {
 	TRACE_CPUPROFILER_EVENT_SCOPE(UGraph::MergeOrCreateIslands);
 	if (Edges.IsEmpty())
@@ -560,14 +613,25 @@ void UGraph::RemoveBulkVertices(const TArray<FGraphVertexHandle>& InHandles)
 					AffectedIslands.Add(Node->GetParentIsland());
 					Island->RemoveVertex(NodeHandle);
 				}
-			}
 
-			// We must remove every edge this node is a part of.
-			for (const FGraphEdgeHandle& EdgeHandle : VertexEdges.FindOrAdd(NodeHandle))
-			{
-				// Don't immediately handle islands. We'll do it later.
-				RemoveEdge(EdgeHandle, false);
+				// We must remove every edge this node is a part of. Need to make a copy of the edges because
+				// otherwise we're modifying the container during iteration over it in ForEachAdjacentVertex.
+				TArray<FGraphEdgeHandle> EdgeCopy;
+				EdgeCopy.Reserve(Node->NumEdges());
+				Node->ForEachAdjacentVertex(
+					[&EdgeCopy](const FGraphVertexHandle& OtherNodeHandle, const FGraphEdgeHandle& EdgeHandle)
+					{
+						EdgeCopy.Add(EdgeHandle);
+					}
+				);
+
+				for (const FGraphEdgeHandle& EdgeHandle : EdgeCopy)
+				{
+					// Don't immediately handle islands. We'll do it later.
+					RemoveEdgeInternal(EdgeHandle, false);
+				}
 			}
+
 		}
 	}
 	
@@ -590,22 +654,24 @@ void UGraph::RemoveBulkVertices(const TArray<FGraphVertexHandle>& InHandles)
 				Node->HandleOnVertexRemoved();
 			}
 			Vertices.Remove(NodeHandle);
-			VertexEdges.Remove(NodeHandle);
 		}
 	}
 }
 
-void UGraph::RemoveEdge(const FGraphEdgeHandle& EdgeHandle, bool bHandleIslands)
+void UGraph::RemoveEdge(const FGraphEdgeHandle& EdgeHandle)
+{
+	RemoveEdgeInternal(EdgeHandle, true);
+}
+
+void UGraph::RemoveEdgeInternal(const FGraphEdgeHandle& EdgeHandle, bool bHandleIslands)
 {
-	TRACE_CPUPROFILER_EVENT_SCOPE(UGraph::RemoveEdge);
+	TRACE_CPUPROFILER_EVENT_SCOPE(UGraph::RemoveEdgeInternal);
 	if (!EdgeHandle.IsValid())
 	{
 		return;
 	}
 
 	TObjectPtr<UGraphEdge> Edge = EdgeHandle.GetEdge();
-	Edges.Remove(EdgeHandle);
-
 	if (!Edge)
 	{
 		return;
@@ -643,6 +709,8 @@ void UGraph::RemoveEdge(const FGraphEdgeHandle& EdgeHandle, bool bHandleIslands)
 		// Note that we can assume that both nodes are in the same island.
 		RemoveOrSplitIsland(IslandHandle.GetIsland());
 	}
+
+	Edges.Remove(EdgeHandle);
 }
 
 void UGraph::RemoveOrSplitIsland(TObjectPtr<UGraphIsland> Island)
@@ -709,9 +777,9 @@ void UGraph::RemoveOrSplitIsland(TObjectPtr<UGraphIsland> Island)
 
 FGraphVertexHandle UGraph::GetCompleteNodeHandle(const FGraphVertexHandle& InHandle) const
 {
-	if (!InHandle.IsValid() || InHandle.HasElement())
+	if (!InHandle.IsValid())
 	{
-		return InHandle;
+		return {};
 	}
 
 	const TObjectPtr<UGraphVertex>* Data = Vertices.Find(InHandle);
@@ -759,6 +827,28 @@ void UGraph::FinalizeVertex(const FGraphVertexHandle& InHandle)
 		IslandHandle = CreateIsland({ InHandle });
 	}
 
+	// Need to iterate through edges to see if any edge was added prior to vertex finalization.
+	// If so, we need to do a merge.
+	TArray<FGraphEdgeHandle> EdgesToMerge;
+	Node->ForEachAdjacentVertex(
+		[&IslandHandle, &EdgesToMerge](const FGraphVertexHandle& NeighborVertexHandle, const FGraphEdgeHandle& EdgeHandle)
+		{
+			if (UGraphVertex* NeighborVertex = NeighborVertexHandle.GetVertex())
+			{
+				const FGraphIslandHandle& NeighborIslandHandle = NeighborVertex->GetParentIsland();
+				if (NeighborIslandHandle.IsComplete() && NeighborIslandHandle != IslandHandle)
+				{
+					EdgesToMerge.Add(EdgeHandle);
+				}
+			}
+		}
+	);
+
+	if (!EdgesToMerge.IsEmpty())
+	{
+		MergeOrCreateIslands(EdgesToMerge);
+	}
+
 	if (TObjectPtr<UGraphIsland> Island = IslandHandle.GetIsland())
 	{
 		Island->HandleOnConnectivityChanged();
diff --git a/Engine/Plugins/Experimental/GameplayGraph/Source/GameplayGraph/Private/Graph/GraphEdge.cpp b/Engine/Plugins/Experimental/GameplayGraph/Source/GameplayGraph/Private/Graph/GraphEdge.cpp
--- a/Engine/Plugins/Experimental/GameplayGraph/Source/GameplayGraph/Private/Graph/GraphEdge.cpp
+++ b/Engine/Plugins/Experimental/GameplayGraph/Source/GameplayGraph/Private/Graph/GraphEdge.cpp
@@ -21,10 +21,14 @@ const FGraphVertexHandle& UGraphEdge::GetOtherNode(const FGraphVertexHandle& InN
 	{
 		return B;
 	}
-	else
+	else if (InNode == B)
 	{
 		return A;
 	}
+	else
+	{
+		return FGraphVertexHandle::Invalid;
+	}
 }
 
 bool UGraphEdge::ContainsNode(const FGraphVertexHandle& InNode) const
diff --git a/Engine/Plugins/Experimental/GameplayGraph/Source/GameplayGraph/Private/Graph/GraphHandle.cpp b/Engine/Plugins/Experimental/GameplayGraph/Source/GameplayGraph/Private/Graph/GraphHandle.cpp
--- a/Engine/Plugins/Experimental/GameplayGraph/Source/GameplayGraph/Private/Graph/GraphHandle.cpp
+++ b/Engine/Plugins/Experimental/GameplayGraph/Source/GameplayGraph/Private/Graph/GraphHandle.cpp
@@ -1,6 +1,8 @@
 // Copyright Epic Games, Inc. All Rights Reserved.
 
 #include "Graph/GraphHandle.h"
+
+#include "Graph/Graph.h"
 #include "Graph/GraphElement.h"
 #include "Graph/GraphEdge.h"
 #include "Graph/GraphIsland.h"
@@ -11,37 +13,21 @@
 
 DEFINE_LOG_CATEGORY(LogGameplayGraph)
 
-FGraphHandle::FGraphHandle()
-{
+FGraphVertexHandle FGraphVertexHandle::Invalid;
+FGraphEdgeHandle FGraphEdgeHandle::Invalid;
+FGraphIslandHandle FGraphIslandHandle::Invalid;
 
-}
-
-FGraphHandle::FGraphHandle(FGraphUniqueIndex InUniqueIndex, TObjectPtr<UGraphElement> InElement)
+FGraphHandle::FGraphHandle(FGraphUniqueIndex InUniqueIndex, UGraph* InGraph)
 	: UniqueIndex(InUniqueIndex)
-	, Element(InElement)
+	, WeakGraph(InGraph)
 {
 }
 
-void FGraphHandle::SetElement(TObjectPtr<UGraphElement> InElement)
-{
-	Element = InElement;
-}
-
-TObjectPtr<UGraphElement> FGraphHandle::GetElement() const
-{
-	return Element.Get();
-}
-
 bool FGraphHandle::IsValid() const
 {
 	return UniqueIndex.IsValid();
 }
 
-bool FGraphHandle::HasElement() const
-{
-	return Element.IsValid();
-}
-
 bool FGraphHandle::IsComplete() const
 {
 	return IsValid() && HasElement();
@@ -65,7 +51,7 @@ bool FGraphHandle::operator<(const FGraphHandle& Other) const
 void FGraphHandle::Clear()
 {
 	UniqueIndex = FGraphUniqueIndex();
-	Element = nullptr;
+	WeakGraph = nullptr;
 }
 
 uint32 GetTypeHash(const FGraphHandle& Handle)
@@ -73,31 +59,62 @@ uint32 GetTypeHash(const FGraphHandle& Handle)
 	return GetTypeHash(Handle.UniqueIndex);
 }
 
-FGraphVertexHandle::FGraphVertexHandle()
+UGraphVertex* FGraphVertexHandle::GetVertex() const
 {
+	if (UGraph* Graph = GetGraph())
+	{
+		return Graph->GetVertices().FindRef(*this);
+	}
+
+	return nullptr;
 }
 
-TObjectPtr<UGraphVertex> FGraphVertexHandle::GetVertex() const
+bool FGraphVertexHandle::HasElement() const
 {
-	return Cast<UGraphVertex>(GetElement());
+	if (UGraph* Graph = GetGraph())
+	{
+		return Graph->GetVertices().Contains(*this);
+	}
+
+	return false;
 }
 
-FGraphEdgeHandle::FGraphEdgeHandle()
+UGraphEdge* FGraphEdgeHandle::GetEdge() const
 {
+	if (UGraph* Graph = GetGraph())
+	{
+		return Graph->GetEdges().FindRef(*this);
+	}
 
+	return nullptr;
 }
 
-TObjectPtr<UGraphEdge> FGraphEdgeHandle::GetEdge() const
+bool FGraphEdgeHandle::HasElement() const
 {
-	return Cast<UGraphEdge>(GetElement());
+	if (UGraph* Graph = GetGraph())
+	{
+		return Graph->GetEdges().Contains(*this);
+	}
+
+	return false;
 }
 
-FGraphIslandHandle::FGraphIslandHandle()
+UGraphIsland* FGraphIslandHandle::GetIsland() const
 {
+	if (UGraph* Graph = GetGraph())
+	{
+		return Graph->GetIslands().FindRef(*this);
+	}
 
+	return nullptr;
 }
 
-TObjectPtr<UGraphIsland> FGraphIslandHandle::GetIsland() const
+bool FGraphIslandHandle::HasElement() const
 {
-	return Cast<UGraphIsland>(GetElement());
+	if (UGraph* Graph = GetGraph())
+	{
+		return Graph->GetIslands().Contains(*this);
+	}
+
+	return false;
 }
\ No newline at end of file
diff --git a/Engine/Plugins/Experimental/GameplayGraph/Source/GameplayGraph/Private/Graph/GraphIsland.cpp b/Engine/Plugins/Experimental/GameplayGraph/Source/GameplayGraph/Private/Graph/GraphIsland.cpp
--- a/Engine/Plugins/Experimental/GameplayGraph/Source/GameplayGraph/Private/Graph/GraphIsland.cpp
+++ b/Engine/Plugins/Experimental/GameplayGraph/Source/GameplayGraph/Private/Graph/GraphIsland.cpp
@@ -9,13 +9,11 @@
 UGraphIsland::UGraphIsland()
 	: Super(EGraphElementType::Island)
 {
-	bPendingDestroy = false;
 }
 
 void UGraphIsland::Destroy()
 {
 	TRACE_CPUPROFILER_EVENT_SCOPE(UGraphIsland::Destroy);
-	bPendingDestroy = true;
 	TSet<FGraphVertexHandle> VertexCopy = Vertices;
 	Vertices.Empty();
 
diff --git a/Engine/Plugins/Experimental/GameplayGraph/Source/GameplayGraph/Public/Graph/Graph.h b/Engine/Plugins/Experimental/GameplayGraph/Source/GameplayGraph/Public/Graph/Graph.h
--- a/Engine/Plugins/Experimental/GameplayGraph/Source/GameplayGraph/Public/Graph/Graph.h
+++ b/Engine/Plugins/Experimental/GameplayGraph/Source/GameplayGraph/Public/Graph/Graph.h
@@ -56,6 +56,13 @@ struct FSerializableGraph
 	TMap<FGraphIslandHandle, FSerializedIslandData> Islands;
 };
 
+struct FEdgeCreationParameters
+{
+	FGraphVertexHandle VertexHandle1;
+	FGraphVertexHandle VertexHandle2;
+	FGraphUniqueIndex EdgeIndex = FGraphUniqueIndex::CreateUniqueIndex();
+};
+
 /**
  * A UGraph is a collection of nodes and edges. This graph representation
  * is meant to be easily integrable into gameplay systems in the Unreal Engine.
@@ -95,13 +102,10 @@ public:
 	FGraphVertexHandle GetCompleteNodeHandle(const FGraphVertexHandle& InHandle) const;
 
 	/** Create a node with the specified subclass, adds it to the graph, and returns a handle to it. */
-	FGraphVertexHandle CreateVertex(FGraphUniqueIndex InUniqueIndex = FGraphUniqueIndex());
-
-	/** Creates an edge between the two nodes. */
-	FGraphEdgeHandle CreateEdge(FGraphVertexHandle Node1, FGraphVertexHandle Node2, FGraphUniqueIndex InUniqueIndex = FGraphUniqueIndex(), bool bAddToIslands = true);
+	FGraphVertexHandle CreateVertex(FGraphUniqueIndex InUniqueIndex = FGraphUniqueIndex::CreateUniqueIndex(false));
 
 	/** Creates edges in bulk. This is more efficient than calling CreateEdge multiple times since we will only try to assign a node to an island once. */
-	void CreateBulkEdges(TArray<TPair<FGraphVertexHandle, FGraphVertexHandle>>&& NodesToConnect);
+	void CreateBulkEdges(TArray<FEdgeCreationParameters>&& NodesToConnect, TArray<FGraphEdgeHandle>* OutEdges = nullptr);
 
 	/** Removes a node from the graph along with any edges that contain it. */
 	void RemoveVertex(const FGraphVertexHandle& NodeHandle);
@@ -109,7 +113,7 @@ public:
 	void RemoveBulkVertices(const TArray<FGraphVertexHandle>& InHandles);
 
 	/** Removes an edge from the graph. */
-	void RemoveEdge(const FGraphEdgeHandle& EdgeHandle, bool bHandleIslands);
+	void RemoveEdge(const FGraphEdgeHandle& EdgeHandle);
 
 	/** This should be called immediately after a node and any relevant edges have been added to the graph. */
 	virtual void FinalizeVertex(const FGraphVertexHandle& InHandle);
@@ -124,11 +128,19 @@ public:
 	int32 NumEdges() const { return Edges.Num(); }
 	int32 NumIslands() const { return Islands.Num(); }
 
+	template<typename TLambda>
+	void ForEachIsland(TLambda&& Lambda)
+	{
+		for (const TPair<FGraphIslandHandle, TObjectPtr<UGraphIsland>>& Kvp : Islands)
+		{
+			Lambda(Kvp.Key, Kvp.Value);
+		}
+	}
+
 	FOnGraphVertexCreated OnVertexCreated;
 	FOnGraphEdgeCreated OnEdgeCreated;
 	FOnGraphIslandCreated OnIslandCreated;
 
-protected:
 	const TMap<FGraphVertexHandle, TObjectPtr<UGraphVertex>>& GetVertices() const { return Vertices; }
 	const TMap<FGraphEdgeHandle, TObjectPtr<UGraphEdge>>& GetEdges() const { return Edges; }
 	const TMap<FGraphIslandHandle, TObjectPtr<UGraphIsland>>& GetIslands() const { return Islands; }
@@ -138,9 +150,6 @@ private:
 	UPROPERTY()
 	TMap<FGraphVertexHandle, TObjectPtr<UGraphVertex>> Vertices;
 
-	/** Keeps track of every edge a node is a part of. The node itself will only track edges it's the "source" node of in directed graphs. */
-	TMap<FGraphVertexHandle, TSet<FGraphEdgeHandle>> VertexEdges;
-
 	UPROPERTY()
 	TMap<FGraphEdgeHandle, TObjectPtr<UGraphEdge>> Edges;
 
@@ -151,6 +160,9 @@ private:
 	/** Creates an island out of a given set of nodes. */
 	FGraphIslandHandle CreateIsland(TArray<FGraphVertexHandle> Nodes, FGraphUniqueIndex InUniqueIndex = FGraphUniqueIndex());
 
+	/** Creates an edge between the two nodes. */
+	FGraphEdgeHandle CreateEdge(FGraphVertexHandle Node1, FGraphVertexHandle Node2, FGraphUniqueIndex InUniqueIndex = FGraphUniqueIndex(), bool bMergeIslands = true);
+
 	/** Adds a node to the graph's node collection and modifies the NextAvailableNodeUniqueIndex as necessary to maintain its validity. */
 	void RegisterVertex(TObjectPtr<UGraphVertex> Node);
 
@@ -161,11 +173,14 @@ private:
 	void RegisterIsland(TObjectPtr<UGraphIsland> Edge);
 
 	/** When we add multiple edges into the graph. This function will ensure that the interactions we make externally are kept to a minimum. */
-	void MergeOrCreateIslands(TArray<FGraphEdgeHandle>&& InEdges);
+	void MergeOrCreateIslands(const TArray<FGraphEdgeHandle>& InEdges);
 
 	/** After a change, this function will remove the island if it's empty or will attempt to split it into two smaller islands. */
 	void RemoveOrSplitIsland(TObjectPtr<UGraphIsland> Island);
 
+	/** Need this so that users aren't able to pass in bHandleIslands = false. */
+	void RemoveEdgeInternal(const FGraphEdgeHandle& EdgeHandle, bool bHandleIslands);
+
 	/** Factory functions for creating an appropriately typed node/edge/island. */
 	virtual TObjectPtr<UGraphVertex> CreateTypedVertex() const;
 	virtual TObjectPtr<UGraphEdge> CreateTypedEdge() const;
diff --git a/Engine/Plugins/Experimental/GameplayGraph/Source/GameplayGraph/Public/Graph/GraphEdge.h b/Engine/Plugins/Experimental/GameplayGraph/Source/GameplayGraph/Public/Graph/GraphEdge.h
--- a/Engine/Plugins/Experimental/GameplayGraph/Source/GameplayGraph/Public/Graph/GraphEdge.h
+++ b/Engine/Plugins/Experimental/GameplayGraph/Source/GameplayGraph/Public/Graph/GraphEdge.h
@@ -40,7 +40,7 @@ public:
 
 	FGraphEdgeHandle Handle() const
 	{
-		return FGraphEdgeHandle{ GetUniqueIndex(), const_cast<UGraphEdge*>(this) };
+		return FGraphEdgeHandle{ GetUniqueIndex(), GetGraph() };
 	}
 
 	FSerializedEdgeData GetSerializedData() const;
diff --git a/Engine/Plugins/Experimental/GameplayGraph/Source/GameplayGraph/Public/Graph/GraphHandle.h b/Engine/Plugins/Experimental/GameplayGraph/Source/GameplayGraph/Public/Graph/GraphHandle.h
--- a/Engine/Plugins/Experimental/GameplayGraph/Source/GameplayGraph/Public/Graph/GraphHandle.h
+++ b/Engine/Plugins/Experimental/GameplayGraph/Source/GameplayGraph/Public/Graph/GraphHandle.h
@@ -6,7 +6,7 @@
 
 #include "GraphHandle.generated.h"
 
-class UGraphElement;
+class UGraph;
 class UGraphVertex;
 class UGraphEdge;
 class UGraphIsland;
@@ -102,20 +102,19 @@ struct GAMEPLAYGRAPH_API FGraphHandle
 {
 	GENERATED_BODY()
 public:
-	FGraphHandle();
-	FGraphHandle(FGraphUniqueIndex InUniqueIndex, TObjectPtr<UGraphElement> InElement);
+	FGraphHandle() = default;
+	FGraphHandle(FGraphUniqueIndex InUniqueIndex, UGraph* InGraph);
+	virtual ~FGraphHandle() = default;
 
 	void Clear();
 
 	/** Whether or not this handle has been initialized. */
 	bool IsValid() const;
-	bool HasElement() const;
 	bool IsComplete() const;
+	virtual bool HasElement() const { return false; }
 
 	FGraphUniqueIndex GetUniqueIndex() const { return UniqueIndex; }
-
-	void SetElement(TObjectPtr<UGraphElement> InElement);
-	TObjectPtr<UGraphElement> GetElement() const;
+	UGraph* GetGraph() const { return WeakGraph.Get(); }
 
 	bool operator==(const FGraphHandle& Other) const;
 	bool operator!=(const FGraphHandle& Other) const;
@@ -129,7 +128,7 @@ private:
 
 	/** Pointer to the graph */
 	UPROPERTY(Transient)
-	TWeakObjectPtr<UGraphElement> Element;
+	TWeakObjectPtr<UGraph> WeakGraph;
 };
 
 USTRUCT()
@@ -137,12 +136,15 @@ struct GAMEPLAYGRAPH_API FGraphVertexHandle : public FGraphHandle
 {
 	GENERATED_BODY()
 
-	FGraphVertexHandle();
-	FGraphVertexHandle(FGraphUniqueIndex InUniqueIndex, TObjectPtr<UGraphElement> InElement = nullptr)
-		: FGraphHandle(InUniqueIndex, InElement)
+	static FGraphVertexHandle Invalid;
+
+	FGraphVertexHandle() = default;
+	FGraphVertexHandle(FGraphUniqueIndex InUniqueIndex, UGraph* InGraph)
+		: FGraphHandle(InUniqueIndex, InGraph)
 	{}
 
-	TObjectPtr<UGraphVertex> GetVertex() const;
+	UGraphVertex* GetVertex() const;
+	virtual bool HasElement() const override;
 };
 
 USTRUCT()
@@ -150,12 +152,15 @@ struct GAMEPLAYGRAPH_API FGraphEdgeHandle : public FGraphHandle
 {
 	GENERATED_BODY()
 
-	FGraphEdgeHandle();
-	FGraphEdgeHandle(FGraphUniqueIndex InUniqueIndex, TObjectPtr<UGraphElement> InElement)
-		: FGraphHandle(InUniqueIndex, InElement)
+	static FGraphEdgeHandle Invalid;
+
+	FGraphEdgeHandle() = default;
+	FGraphEdgeHandle(FGraphUniqueIndex InUniqueIndex, UGraph* InGraph)
+		: FGraphHandle(InUniqueIndex, InGraph)
 	{}
 
-	TObjectPtr<UGraphEdge> GetEdge() const;
+	UGraphEdge* GetEdge() const;
+	virtual bool HasElement() const override;
 };
 
 USTRUCT()
@@ -163,10 +168,13 @@ struct GAMEPLAYGRAPH_API FGraphIslandHandle : public FGraphHandle
 {
 	GENERATED_BODY()
 
-	FGraphIslandHandle();
-	FGraphIslandHandle(FGraphUniqueIndex InUniqueIndex, TObjectPtr<UGraphElement> InElement)
-		: FGraphHandle(InUniqueIndex, InElement)
+	static FGraphIslandHandle Invalid;
+
+	FGraphIslandHandle() = default;
+	FGraphIslandHandle(FGraphUniqueIndex InUniqueIndex, UGraph* InGraph)
+		: FGraphHandle(InUniqueIndex, InGraph)
 	{}
 
-	TObjectPtr<UGraphIsland> GetIsland() const;
-};
\ No newline at end of file
+	UGraphIsland* GetIsland() const;
+	virtual bool HasElement() const override;
+};
diff --git a/Engine/Plugins/Experimental/GameplayGraph/Source/GameplayGraph/Public/Graph/GraphIsland.h b/Engine/Plugins/Experimental/GameplayGraph/Source/GameplayGraph/Public/Graph/GraphIsland.h
--- a/Engine/Plugins/Experimental/GameplayGraph/Source/GameplayGraph/Public/Graph/GraphIsland.h
+++ b/Engine/Plugins/Experimental/GameplayGraph/Source/GameplayGraph/Public/Graph/GraphIsland.h
@@ -60,10 +60,9 @@ public:
 
 	FGraphIslandHandle Handle() const
 	{
-		return FGraphIslandHandle{ GetUniqueIndex(), const_cast<UGraphIsland*>(this) };
+		return FGraphIslandHandle{ GetUniqueIndex(), GetGraph() };
 	}
 
-	bool IsPendingDestroy() const { return bPendingDestroy; }
 	bool IsEmpty() const { return Vertices.IsEmpty(); }
 	const TSet<FGraphVertexHandle>& GetVertices() const { return Vertices; }
 	int32 Num() const { return Vertices.Num(); }
@@ -107,9 +106,6 @@ private:
 	UPROPERTY(SaveGame)
 	TSet<FGraphVertexHandle> Vertices;
 
-	UPROPERTY(Transient)
-	bool bPendingDestroy;
-
 	UPROPERTY(Transient)
 	EGraphIslandOperations AllowedOperations = EGraphIslandOperations::All;
 };
\ No newline at end of file
diff --git a/Engine/Plugins/Experimental/GameplayGraph/Source/GameplayGraph/Public/Graph/GraphVertex.h b/Engine/Plugins/Experimental/GameplayGraph/Source/GameplayGraph/Public/Graph/GraphVertex.h
--- a/Engine/Plugins/Experimental/GameplayGraph/Source/GameplayGraph/Public/Graph/GraphVertex.h
+++ b/Engine/Plugins/Experimental/GameplayGraph/Source/GameplayGraph/Public/Graph/GraphVertex.h
@@ -18,11 +18,12 @@ public:
 
 	FGraphVertexHandle Handle() const
 	{
-		return FGraphVertexHandle{ GetUniqueIndex(), const_cast<UGraphVertex*>(this)};
+		return FGraphVertexHandle{ GetUniqueIndex(), GetGraph()};
 	}
 
 	bool HasEdgeTo(const FGraphVertexHandle& Other) const;
 	const FGraphEdgeHandle* FindEdgeTo(const FGraphVertexHandle& Other) const;
+	int32 NumEdges() const { return Edges.Num(); }
 
 	const FGraphIslandHandle& GetParentIsland() const { return ParentIsland; }
 
@@ -41,6 +42,8 @@ public:
 			}
 		}
 	}
+	
+	const TMap<FGraphVertexHandle, FGraphEdgeHandle>& GetEdges() const { return Edges; }
 
 	FOnGraphVertexRemoved OnVertexRemoved;
 	FOnGraphVertexParentIslandSet OnParentIslandSet;
diff --git a/Engine/Plugins/Experimental/GameplayGraph/Tests/GameplayGraphTests.Build.cs b/Engine/Plugins/Experimental/GameplayGraph/Tests/GameplayGraphTests.Build.cs
new file mode 100644
--- /dev/null
+++ b/Engine/Plugins/Experimental/GameplayGraph/Tests/GameplayGraphTests.Build.cs
@@ -0,0 +1,21 @@
+// Copyright Epic Games, Inc. All Rights Reserved.
+
+using UnrealBuildTool;
+
+public class GameplayGraphTests : TestModuleRules
+{
+	public GameplayGraphTests(ReadOnlyTargetRules Target) : base(Target)
+	{
+		PrivateDependencyModuleNames.AddRange(
+			new string[] {
+				"Core",
+				"CoreUObject",
+				"GameplayGraph"
+			});
+
+		UpdateBuildGraphPropertiesFile(new Metadata() {
+			TestName = "GameplayGraph",
+			TestShortName = "GameplayGraph"
+		});
+	}
+}
\ No newline at end of file
diff --git a/Engine/Plugins/Experimental/GameplayGraph/Tests/GameplayGraphTests.Target.cs b/Engine/Plugins/Experimental/GameplayGraph/Tests/GameplayGraphTests.Target.cs
new file mode 100644
--- /dev/null
+++ b/Engine/Plugins/Experimental/GameplayGraph/Tests/GameplayGraphTests.Target.cs
@@ -0,0 +1,11 @@
+// Copyright Epic Games, Inc. All Rights Reserved.
+
+using UnrealBuildTool;
+
+[SupportedPlatforms(UnrealPlatformClass.All)]
+public class GameplayGraphTestsTarget : TestTargetRules
+{
+	public GameplayGraphTestsTarget(TargetInfo Target) : base(Target)
+	{
+	}
+}
diff --git a/Engine/Plugins/Experimental/GameplayGraph/Tests/Private/GraphEdgeTests.cpp b/Engine/Plugins/Experimental/GameplayGraph/Tests/Private/GraphEdgeTests.cpp
new file mode 100644
--- /dev/null
+++ b/Engine/Plugins/Experimental/GameplayGraph/Tests/Private/GraphEdgeTests.cpp
@@ -0,0 +1,117 @@
+// Copyright Epic Games, Inc. All Rights Reserved.
+#include "TestHarness.h"
+#include "TestGraphBuilder.h"
+
+TEST_CASE_METHOD(FTestGraphBuilder, "Graph::Edge::Create", "[graph][edge]")
+{
+	PopulateVertices(3, true);
+
+	FEdgeCreationParameters Params;
+	Params.VertexHandle1 = VertexHandles[1];
+	Params.VertexHandle2 = VertexHandles[0];
+	Params.EdgeIndex = FGraphUniqueIndex{ FGuid{ 1, 1, 1, 1 } };
+
+	TArray<FGraphEdgeHandle> Edges;
+	Graph->CreateBulkEdges({ Params }, &Edges);
+	REQUIRE(Edges.Num() == 1);
+
+	SECTION("Graph Edge Properties")
+	{
+		CHECK(Edges[0].HasElement());
+		CHECK(Edges[0].IsValid());
+
+		UGraphEdge* E = Edges[0].GetEdge();
+		REQUIRE(E != nullptr);
+		CHECK(E->NodeA() == VertexHandles[0]);
+		CHECK(E->NodeB() == VertexHandles[1]);
+		CHECK(E->Handle() == Edges[0]);
+		CHECK(E->ContainsNode(VertexHandles[0]) == true);
+		CHECK(E->ContainsNode(VertexHandles[1]) == true);
+		CHECK(E->ContainsNode(VertexHandles[2]) == false);
+		CHECK(E->GetOtherNode(VertexHandles[0]) == VertexHandles[1]);
+		CHECK(E->GetOtherNode(VertexHandles[1]) == VertexHandles[0]);
+		CHECK(E->GetOtherNode(VertexHandles[2]) == FGraphVertexHandle{});
+
+		FSerializedEdgeData Serialized = E->GetSerializedData();
+		CHECK(Serialized.Node1 == VertexHandles[0]);
+		CHECK(Serialized.Node2 == VertexHandles[1]);
+	}
+
+	SECTION("Remove Edge")
+	{
+		CHECK(Edges[0].IsValid() == true);
+		CHECK(Edges[0].IsComplete() == true);
+		CHECK(Edges[0].HasElement() == true);
+		CHECK(Edges[0].GetUniqueIndex() == Params.EdgeIndex);
+		CHECK(Edges[0].GetGraph() == Graph);
+
+		Graph->RemoveEdge(Edges[0]);
+
+		CHECK(Edges[0].IsValid() == true);
+		CHECK(Edges[0].IsComplete() == false);
+		CHECK(Edges[0].HasElement() == false);
+		CHECK(Edges[0].GetUniqueIndex() == Params.EdgeIndex);
+		CHECK(Edges[0].GetGraph() == Graph);
+	}
+}
+
+TEST_CASE_METHOD(FTestGraphBuilder, "Graph::Edge::Remove Vertex", "[graph][edge]")
+{
+	PopulateVertices(3, true);
+
+	FEdgeCreationParameters Params;
+	Params.VertexHandle1 = VertexHandles[1];
+	Params.VertexHandle2 = VertexHandles[0];
+	Params.EdgeIndex = FGraphUniqueIndex{ FGuid{ 1, 1, 1, 1 } };
+
+	TArray<FGraphEdgeHandle> Edges;
+	Graph->CreateBulkEdges({ Params }, &Edges);
+	REQUIRE(Edges.Num() == 1);
+
+	CHECK(Edges[0].IsComplete() == true);
+	Graph->RemoveVertex(VertexHandles[0]);
+	CHECK(Edges[0].IsComplete() == false);
+}
+
+TEST_CASE_METHOD(FTestGraphBuilder, "Graph::Edge::Remove Edge", "[graph][edge]")
+{
+	PopulateVertices(3, true);
+
+	FEdgeCreationParameters Params;
+	Params.VertexHandle1 = VertexHandles[1];
+	Params.VertexHandle2 = VertexHandles[0];
+	Params.EdgeIndex = FGraphUniqueIndex{ FGuid{ 1, 1, 1, 1 } };
+
+	TArray<FGraphEdgeHandle> Edges;
+	Graph->CreateBulkEdges({ Params }, &Edges);
+	REQUIRE(Edges.Num() == 1);
+
+	SECTION("With islands")
+	{
+		CHECK(Edges[0].IsComplete() == true);
+		Graph->RemoveEdge(Edges[0]);
+		CHECK(Edges[0].IsComplete() == false);
+	}
+}
+
+TEST_CASE_METHOD(FTestGraphBuilder, "Graph::Edge::Remove island", "[graph][edge]")
+{
+	PopulateVertices(3, true);
+
+	FEdgeCreationParameters Params;
+	Params.VertexHandle1 = VertexHandles[1];
+	Params.VertexHandle2 = VertexHandles[0];
+	Params.EdgeIndex = FGraphUniqueIndex{ FGuid{ 1, 1, 1, 1 } };
+
+	TArray<FGraphEdgeHandle> Edges;
+	CHECK(Graph->GetIslands().Num() == 3);
+	Graph->CreateBulkEdges({ Params }, &Edges);
+	REQUIRE(Edges.Num() == 1);
+	REQUIRE(Graph->GetIslands().Num() == 2);
+
+	FGraphIslandHandle IslandHandle = Graph->GetIslands().CreateConstIterator()->Key;
+
+	CHECK(Edges[0].IsComplete() == true);
+	Graph->RemoveIsland(IslandHandle);
+	CHECK(Edges[0].IsComplete() == false);
+}
\ No newline at end of file
diff --git a/Engine/Plugins/Experimental/GameplayGraph/Tests/Private/GraphElementTests.cpp b/Engine/Plugins/Experimental/GameplayGraph/Tests/Private/GraphElementTests.cpp
new file mode 100644
--- /dev/null
+++ b/Engine/Plugins/Experimental/GameplayGraph/Tests/Private/GraphElementTests.cpp
@@ -0,0 +1,65 @@
+// Copyright Epic Games, Inc. All Rights Reserved.
+#include "TestGraphElement.h"
+#include "TestHarness.h"
+
+#include "Graph/Graph.h"
+#include "Graph/GraphEdge.h"
+#include "Graph/GraphIsland.h"
+#include "Graph/GraphVertex.h"
+
+TEST_CASE("Graph::Element::Element Type", "[graph][element]")
+{
+	SECTION("Vertex")
+	{
+		UGraphVertex* Vertex = NewObject<UGraphVertex>();
+		CHECK(Vertex->GetElementType() == EGraphElementType::Node);
+	}
+
+	SECTION("Edge")
+	{
+		UGraphEdge* Edge = NewObject<UGraphEdge>();
+		CHECK(Edge->GetElementType() == EGraphElementType::Edge);
+	}
+
+	SECTION("Island")
+	{
+		UGraphIsland* Island = NewObject<UGraphIsland>();
+		CHECK(Island->GetElementType() == EGraphElementType::Island);
+	}
+}
+
+TEST_CASE("Graph::Element::Unique Index", "[graph][properties]")
+{
+	UTestGraphElement* Element = NewObject<UTestGraphElement>();
+	SECTION("Non-Temp")
+	{
+		FGraphUniqueIndex Test{ FGuid::NewGuid(), false };
+		Element->PublicSetUniqueIndex(Test);
+		CHECK(Element->PublicGetUniqueIndex() == Test);
+	}
+
+	SECTION("Non-Temp")
+	{
+		FGraphUniqueIndex Test{ FGuid::NewGuid(), true };
+		Element->PublicSetUniqueIndex(Test);
+		CHECK(Element->PublicGetUniqueIndex() == Test);
+	}
+}
+
+TEST_CASE("Graph::Element::Graph Pointer", "[graph][properties]")
+{
+	UTestGraphElement* Element = NewObject<UTestGraphElement>();
+
+	SECTION("Set graph")
+	{
+		TObjectPtr<UGraph> TestGraph = NewObject<UGraph>();
+		Element->PublicSetParentGraph(TestGraph);
+		CHECK(Element->PublicGetGraph() == TestGraph);
+	}
+
+	SECTION("Set nullptr")
+	{
+		Element->PublicSetParentGraph(nullptr);
+		CHECK(Element->PublicGetGraph() == nullptr);
+	}
+}
diff --git a/Engine/Plugins/Experimental/GameplayGraph/Tests/Private/GraphHandleTests.cpp b/Engine/Plugins/Experimental/GameplayGraph/Tests/Private/GraphHandleTests.cpp
new file mode 100644
--- /dev/null
+++ b/Engine/Plugins/Experimental/GameplayGraph/Tests/Private/GraphHandleTests.cpp
@@ -0,0 +1,156 @@
+// Copyright Epic Games, Inc. All Rights Reserved.
+#include "TestHarness.h"
+
+#include "Graph/Graph.h"
+#include "Graph/GraphHandle.h"
+
+TEST_CASE("Graph::Unique Index::Constructor", "[graph][uniqueidx]")
+{
+	SECTION("Temp Only Constructor")
+	{
+		FGraphUniqueIndex Test1{false};
+		CHECK(Test1.IsTemporary() == false);
+		CHECK(Test1.IsValid() == false);
+
+		FGraphUniqueIndex Test2{true};
+		CHECK(Test2.IsTemporary() == true);
+		CHECK(Test2.IsValid() == false);
+	}
+
+	SECTION("Guid + Temp Constructor")
+	{
+		FGuid Guid1{1, 1, 1, 1};
+		FGraphUniqueIndex Test1{Guid1, false};
+		CHECK(Test1.ToString() == Guid1.ToString());
+		CHECK(Test1.IsTemporary() == false);
+		CHECK(Test1.IsValid() == true);
+
+		FGuid Guid2{2, 2, 2, 2};
+		FGraphUniqueIndex Test2{Guid2, true};
+		CHECK(Test2.ToString() == Guid2.ToString());
+		CHECK(Test2.IsTemporary() == true);
+		CHECK(Test2.IsValid() == true);
+	}
+
+	SECTION("Create Unique Index")
+	{
+		FGraphUniqueIndex Test1 = FGraphUniqueIndex::CreateUniqueIndex(false);
+		CHECK(Test1.IsTemporary() == false);
+		CHECK(Test1.IsValid() == true);
+
+		FGraphUniqueIndex Test2 = FGraphUniqueIndex::CreateUniqueIndex(true);
+		CHECK(Test2.IsTemporary() == true);
+		CHECK(Test2.IsValid() == true);
+	}
+}
+
+TEST_CASE("Graph::Unique Index::NextUniqueIndex", "[graph][uniqueidx]")
+{
+	SECTION("Temporary")
+	{
+		FGraphUniqueIndex Test1 = FGraphUniqueIndex::CreateUniqueIndex(true);
+		FGraphUniqueIndex Test2 = Test1.NextUniqueIndex();
+		CHECK(Test2.IsTemporary() == true);
+		CHECK(Test2.IsValid() == true);
+	}
+
+	SECTION("Non-Temporary")
+	{
+		FGraphUniqueIndex Test1 = FGraphUniqueIndex::CreateUniqueIndex(false);
+		FGraphUniqueIndex Test2 = Test1.NextUniqueIndex();
+		CHECK(Test2.IsTemporary() == false);
+		CHECK(Test2.IsValid() == true);
+	}
+}
+
+TEST_CASE("Graph::Unique Index::Equality", "[graph][uniqueidx]")
+{
+	SECTION("Equality")
+	{
+		FGuid Guid1{1, 1, 1, 1};
+		FGraphUniqueIndex Test1{Guid1};
+		FGraphUniqueIndex Test2{Guid1};
+		CHECK(Test1 == Test2);
+	}
+
+	SECTION("InEquality")
+	{
+		FGuid Guid1{1, 1, 1, 1};
+		FGuid Guid2{2, 2, 2, 2};
+		FGraphUniqueIndex Test1{Guid1};
+		FGraphUniqueIndex Test2{Guid2};
+		CHECK(Test1 != Test2);
+	}
+}
+
+TEST_CASE("Graph::Unique Index::Comparison", "[graph][uniqueidx]")
+{
+	SECTION("Less")
+	{
+		FGuid Guid1{1, 1, 1, 1};
+		FGuid Guid2{2, 2, 2, 2};
+		FGraphUniqueIndex Test1{Guid1};
+		FGraphUniqueIndex Test2{Guid2};
+		CHECK(Test1 < Test2);
+	}
+}
+
+TEST_CASE("Graph::Unique Index::Copy", "[graph][uniqueidx]")
+{
+	FGraphUniqueIndex Test1 = FGraphUniqueIndex::CreateUniqueIndex(true);
+	FGraphUniqueIndex Test2 = Test1;
+	CHECK(Test1 == Test2);
+}
+
+TEST_CASE("Graph::Handle::Empty Constructor", "[graph][handle]")
+{
+	FGraphHandle Handle;
+	CHECK(Handle.IsValid() == false);
+	CHECK(Handle.HasElement() == false);
+	CHECK(Handle.IsComplete() == false);
+	CHECK(Handle.GetUniqueIndex() == FGraphUniqueIndex{});
+	CHECK(Handle.GetGraph() == nullptr);
+}
+
+TEST_CASE("Graph::Handle::Constructor", "[graph][handle]")
+{
+	UGraph* Graph = NewObject<UGraph>();
+	REQUIRE(Graph != nullptr);
+	FGuid Guid1{1, 1, 1, 1};
+	FGraphUniqueIndex Test1{Guid1};
+	FGraphHandle Handle{Test1, Graph};
+
+	CHECK(Handle.IsValid() == true);
+	CHECK(Handle.HasElement() == false);
+	CHECK(Handle.IsComplete() == false);
+	CHECK(Handle.GetUniqueIndex() == Test1);
+	CHECK(Handle.GetGraph() == Graph);
+}
+
+TEST_CASE("Graph::Handle::Equality", "[graph][handle]")
+{
+	UGraph* Graph = NewObject<UGraph>();
+	REQUIRE(Graph != nullptr);
+	FGuid Guid1{1, 1, 1, 1};
+	FGuid Guid2{2, 2, 2, 2};
+	FGraphUniqueIndex Test1{Guid1};
+	FGraphUniqueIndex Test2{Guid2};
+
+	FGraphHandle Handle1{Test1, Graph};
+	FGraphHandle Handle2{Test2, Graph};
+	FGraphHandle Handle3{Test2, Graph};
+	CHECK(Handle1 != Handle2);
+	CHECK(Handle1 != Handle3);
+	CHECK(Handle2 == Handle3);
+}
+
+TEST_CASE("Graph::Handle::Copy", "[graph][handle]")
+{
+	UGraph* Graph = NewObject<UGraph>();
+	REQUIRE(Graph != nullptr);
+	FGuid Guid1{1, 1, 1, 1};
+	FGraphUniqueIndex Test1{Guid1};
+	FGraphHandle Handle1{Test1, Graph};
+	FGraphHandle Handle2 = Handle1;
+	CHECK(Handle1 == Handle2);
+}
\ No newline at end of file
diff --git a/Engine/Plugins/Experimental/GameplayGraph/Tests/Private/GraphIslandTests.cpp b/Engine/Plugins/Experimental/GameplayGraph/Tests/Private/GraphIslandTests.cpp
new file mode 100644
--- /dev/null
+++ b/Engine/Plugins/Experimental/GameplayGraph/Tests/Private/GraphIslandTests.cpp
@@ -0,0 +1,887 @@
+// Copyright Epic Games, Inc. All Rights Reserved.
+#include "TestHarness.h"
+#include "TestGraphBuilder.h"
+
+TEST_CASE_METHOD(FTestGraphBuilder, "Graph::Island::Create::Single Island", "[graph][island]")
+{
+	PopulateVertices(2, true);
+	BuildLinearEdges(2);
+	FinalizeEdges();
+
+	REQUIRE(IslandHandles.Num() == 1);
+
+	UGraphIsland* Island = IslandHandles[0].GetIsland();
+	REQUIRE(Island != nullptr);
+
+	CHECK(Island->Handle() == IslandHandles[0]);
+	CHECK(Island->IsEmpty() == false);
+	CHECK(Island->GetVertices().Contains(VertexHandles[0]) == true);
+	CHECK(Island->GetVertices().Contains(VertexHandles[1]) == true);
+	CHECK(Island->Num() == 2);
+
+	IslandVertexParentIslandSanityCheck(IslandHandles[0]);
+}
+
+TEST_CASE_METHOD(FTestGraphBuilder, "Graph::Island::Create::Two Islands", "[graph][island]")
+{
+	PopulateVertices(4, false);
+	BuildLinearEdges(2);
+	FinalizeEdges();
+	REQUIRE(IslandHandles.Num() == 2);
+
+	FGraphIslandHandle IslandHandle1;
+	{
+		UGraphVertex* Vertex = VertexHandles[0].GetVertex();
+		REQUIRE(Vertex != nullptr);
+		IslandHandle1 = Vertex->GetParentIsland();
+	}
+
+	UGraphIsland* Island1 = IslandHandle1.GetIsland();
+	REQUIRE(Island1 != nullptr);
+	{
+		CHECK(Island1->Handle() == IslandHandle1);
+		CHECK(Island1->IsEmpty() == false);
+		CHECK(Island1->GetVertices().Contains(VertexHandles[0]) == true);
+		CHECK(Island1->GetVertices().Contains(VertexHandles[1]) == true);
+		CHECK(Island1->GetVertices().Contains(VertexHandles[2]) == false);
+		CHECK(Island1->GetVertices().Contains(VertexHandles[3]) == false);
+		CHECK(Island1->Num() == 2);
+	}
+
+	FGraphIslandHandle IslandHandle2;
+	{
+		UGraphVertex* Vertex = VertexHandles[2].GetVertex();
+		REQUIRE(Vertex != nullptr);
+		IslandHandle2 = Vertex->GetParentIsland();
+	}
+	CHECK(IslandHandle1 != IslandHandle2);
+	UGraphIsland* Island2 = IslandHandle2.GetIsland();
+	REQUIRE(Island2 != nullptr);
+	{
+		CHECK(Island2->Handle() == IslandHandle2);
+		CHECK(Island2->IsEmpty() == false);
+		CHECK(Island2->GetVertices().Contains(VertexHandles[0]) == false);
+		CHECK(Island2->GetVertices().Contains(VertexHandles[1]) == false);
+		CHECK(Island2->GetVertices().Contains(VertexHandles[2]) == true);
+		CHECK(Island2->GetVertices().Contains(VertexHandles[3]) == true);
+		CHECK(Island2->Num() == 2);
+	}
+
+	IslandVertexParentIslandSanityCheck(IslandHandle1);
+	IslandVertexParentIslandSanityCheck(IslandHandle2);
+}
+
+TEST_CASE_METHOD(FTestGraphBuilder, "Graph::Island::Bulk Create Edges", "[graph][island]")
+{
+	PopulateVertices(10, false);
+	BuildLinearEdges(2);
+
+	{
+		TArray<FEdgeCreationParameters> AllEdgeParams;
+		{
+			FEdgeCreationParameters Params;
+			Params.VertexHandle1 = VertexHandles[1];
+			Params.VertexHandle2 = VertexHandles[2];
+			AllEdgeParams.Add(Params);
+		}
+
+		{
+			FEdgeCreationParameters Params;
+			Params.VertexHandle1 = VertexHandles[1];
+			Params.VertexHandle2 = VertexHandles[3];
+			AllEdgeParams.Add(Params);
+		}
+
+		{
+			FEdgeCreationParameters Params;
+			Params.VertexHandle1 = VertexHandles[4];
+			Params.VertexHandle2 = VertexHandles[9];
+			AllEdgeParams.Add(Params);
+		}
+
+		{
+			FEdgeCreationParameters Params;
+			Params.VertexHandle1 = VertexHandles[5];
+			Params.VertexHandle2 = VertexHandles[6];
+			AllEdgeParams.Add(Params);
+		}
+
+		{
+			FEdgeCreationParameters Params;
+			Params.VertexHandle1 = VertexHandles[4];
+			Params.VertexHandle2 = VertexHandles[7];
+			AllEdgeParams.Add(Params);
+		}
+
+		Graph->CreateBulkEdges(MoveTemp(AllEdgeParams));
+	}
+
+	FinalizeEdges();
+
+	REQUIRE(IslandHandles.Num() == 2);
+
+	FGraphIslandHandle IslandHandle1;
+	{
+		UGraphVertex* Vertex = VertexHandles[0].GetVertex();
+		REQUIRE(Vertex != nullptr);
+		IslandHandle1 = Vertex->GetParentIsland();
+	}
+
+	UGraphIsland* Island1 = IslandHandle1.GetIsland();
+	REQUIRE(Island1 != nullptr);
+	{
+		CHECK(Island1->Handle() == IslandHandle1);
+		CHECK(Island1->IsEmpty() == false);
+		CHECK(Island1->GetVertices().Contains(VertexHandles[0]) == true);
+		CHECK(Island1->GetVertices().Contains(VertexHandles[1]) == true);
+		CHECK(Island1->GetVertices().Contains(VertexHandles[2]) == true);
+		CHECK(Island1->GetVertices().Contains(VertexHandles[3]) == true);
+		CHECK(Island1->Num() == 4);
+	}
+
+	FGraphIslandHandle IslandHandle2;
+	{
+		UGraphVertex* Vertex = VertexHandles[4].GetVertex();
+		REQUIRE(Vertex != nullptr);
+		IslandHandle2 = Vertex->GetParentIsland();
+	}
+	CHECK(IslandHandle1 != IslandHandle2);
+	UGraphIsland* Island2 = IslandHandle2.GetIsland();
+	REQUIRE(Island2 != nullptr);
+	{
+		CHECK(Island2->Handle() == IslandHandle2);
+		CHECK(Island2->IsEmpty() == false);
+		CHECK(Island2->GetVertices().Contains(VertexHandles[4]) == true);
+		CHECK(Island2->GetVertices().Contains(VertexHandles[5]) == true);
+		CHECK(Island2->GetVertices().Contains(VertexHandles[6]) == true);
+		CHECK(Island2->GetVertices().Contains(VertexHandles[7]) == true);
+		CHECK(Island2->GetVertices().Contains(VertexHandles[8]) == true);
+		CHECK(Island2->GetVertices().Contains(VertexHandles[9]) == true);
+		CHECK(Island2->Num() == 6);
+	}
+
+	IslandVertexParentIslandSanityCheck(IslandHandle1);
+	IslandVertexParentIslandSanityCheck(IslandHandle2);
+}
+
+TEST_CASE_METHOD(FTestGraphBuilder, "Graph::Island::Remove Vertex::Still Connected", "[graph][island]")
+{
+	PopulateVertices(4, true);
+	BuildFullyConnectedEdges(4);
+
+	REQUIRE(IslandHandles.Num() == 1);
+
+	UGraphIsland* Island1 = IslandHandles[0].GetIsland();
+	REQUIRE(Island1 != nullptr);
+	IslandVertexParentIslandSanityCheck(IslandHandles[0]);
+
+	{
+		CHECK(Island1->Num() == 4);
+		CHECK(Island1->GetVertices().Contains(VertexHandles[0]) == true);
+		CHECK(Island1->GetVertices().Contains(VertexHandles[1]) == true);
+		CHECK(Island1->GetVertices().Contains(VertexHandles[2]) == true);
+		CHECK(Island1->GetVertices().Contains(VertexHandles[3]) == true);
+	}
+
+	Graph->RemoveVertex(VertexHandles[0]);
+	CHECK(Graph->NumIslands() == 1);
+	IslandVertexParentIslandSanityCheck(IslandHandles[0]);
+
+	{
+		CHECK(Island1->Num() == 3);
+		CHECK(Island1->GetVertices().Contains(VertexHandles[0]) == false);
+		CHECK(Island1->GetVertices().Contains(VertexHandles[1]) == true);
+		CHECK(Island1->GetVertices().Contains(VertexHandles[2]) == true);
+		CHECK(Island1->GetVertices().Contains(VertexHandles[3]) == true);
+	}
+
+	Graph->RemoveVertex(VertexHandles[1]);
+	CHECK(Graph->NumIslands() == 1);
+	IslandVertexParentIslandSanityCheck(IslandHandles[0]);
+
+	{
+		CHECK(Island1->Num() == 2);
+		CHECK(Island1->GetVertices().Contains(VertexHandles[0]) == false);
+		CHECK(Island1->GetVertices().Contains(VertexHandles[1]) == false);
+		CHECK(Island1->GetVertices().Contains(VertexHandles[2]) == true);
+		CHECK(Island1->GetVertices().Contains(VertexHandles[3]) == true);
+	}
+
+	Graph->RemoveVertex(VertexHandles[2]);
+	CHECK(Graph->NumIslands() == 1);
+	IslandVertexParentIslandSanityCheck(IslandHandles[0]);
+
+	{
+		CHECK(Island1->Num() == 1);
+		CHECK(Island1->GetVertices().Contains(VertexHandles[0]) == false);
+		CHECK(Island1->GetVertices().Contains(VertexHandles[1]) == false);
+		CHECK(Island1->GetVertices().Contains(VertexHandles[2]) == false);
+		CHECK(Island1->GetVertices().Contains(VertexHandles[3]) == true);
+	}
+
+	IslandVertexParentIslandSanityCheck(IslandHandles[0]);
+	CHECK(Graph->NumIslands() == 1);
+}
+
+TEST_CASE_METHOD(FTestGraphBuilder, "Graph::Island::Remove Vertex::Split", "[graph][island]")
+{
+	PopulateVertices(4, true);
+	BuildLinearEdges(4);
+	REQUIRE(IslandHandles.Num() == 1);
+	CHECK(Graph->NumIslands() == 1);
+
+	Graph->RemoveVertex(VertexHandles[1]);
+	FinalizeEdges();
+	REQUIRE(IslandHandles.Num() == 2);
+	CHECK(Graph->NumIslands() == 2);
+
+	FGraphIslandHandle IslandHandle1;
+	{
+		UGraphVertex* Vertex = VertexHandles[0].GetVertex();
+		REQUIRE(Vertex != nullptr);
+		IslandHandle1 = Vertex->GetParentIsland();
+	}
+
+	UGraphIsland* Island1 = IslandHandle1.GetIsland();
+	REQUIRE(Island1 != nullptr);
+	{
+		CHECK(Island1->Handle() == IslandHandle1);
+		CHECK(Island1->IsEmpty() == false);
+		CHECK(Island1->GetVertices().Contains(VertexHandles[0]) == true);
+		CHECK(Island1->Num() == 1);
+	}
+
+	FGraphIslandHandle IslandHandle2;
+	{
+		UGraphVertex* Vertex = VertexHandles[2].GetVertex();
+		REQUIRE(Vertex != nullptr);
+		IslandHandle2 = Vertex->GetParentIsland();
+	}
+	CHECK(IslandHandle1 != IslandHandle2);
+	UGraphIsland* Island2 = IslandHandle2.GetIsland();
+	REQUIRE(Island2 != nullptr);
+	{
+		CHECK(Island2->Handle() == IslandHandle2);
+		CHECK(Island2->IsEmpty() == false);
+		CHECK(Island2->GetVertices().Contains(VertexHandles[2]) == true);
+		CHECK(Island2->GetVertices().Contains(VertexHandles[3]) == true);
+		CHECK(Island2->Num() == 2);
+	}
+
+	IslandVertexParentIslandSanityCheck(IslandHandle1);
+	IslandVertexParentIslandSanityCheck(IslandHandle2);
+}
+
+TEST_CASE_METHOD(FTestGraphBuilder, "Graph::Island::Remove Edge::Still Connected", "[graph][island]")
+{
+	PopulateVertices(4, true);
+	BuildFullyConnectedEdges(4);
+
+	REQUIRE(IslandHandles.Num() == 1);
+
+	UGraphIsland* Island1 = IslandHandles[0].GetIsland();
+	REQUIRE(Island1 != nullptr);
+	IslandVertexParentIslandSanityCheck(IslandHandles[0]);
+
+	{
+		CHECK(Island1->Num() == 4);
+		CHECK(Island1->GetVertices().Contains(VertexHandles[0]) == true);
+		CHECK(Island1->GetVertices().Contains(VertexHandles[1]) == true);
+		CHECK(Island1->GetVertices().Contains(VertexHandles[2]) == true);
+		CHECK(Island1->GetVertices().Contains(VertexHandles[3]) == true);
+	}
+
+	TArray<FGraphEdgeHandle> Edges = GetEdgesForVertex(VertexHandles[0]);
+	CHECK(Edges.Num() == 3);
+
+	Graph->RemoveEdge(Edges[0]);
+	IslandVertexParentIslandSanityCheck(IslandHandles[0]);
+	{
+		CHECK(Island1->Num() == 4);
+		CHECK(Island1->GetVertices().Contains(VertexHandles[0]) == true);
+		CHECK(Island1->GetVertices().Contains(VertexHandles[1]) == true);
+		CHECK(Island1->GetVertices().Contains(VertexHandles[2]) == true);
+		CHECK(Island1->GetVertices().Contains(VertexHandles[3]) == true);
+	}
+
+	Graph->RemoveEdge(Edges[1]);
+	IslandVertexParentIslandSanityCheck(IslandHandles[0]);
+	{
+		CHECK(Island1->Num() == 4);
+		CHECK(Island1->GetVertices().Contains(VertexHandles[0]) == true);
+		CHECK(Island1->GetVertices().Contains(VertexHandles[1]) == true);
+		CHECK(Island1->GetVertices().Contains(VertexHandles[2]) == true);
+		CHECK(Island1->GetVertices().Contains(VertexHandles[3]) == true);
+	}
+}
+
+TEST_CASE_METHOD(FTestGraphBuilder, "Graph::Island::Remove Edge::Split", "[graph][island]")
+{
+	PopulateVertices(5, true);
+	BuildLinearEdges(5);
+
+	REQUIRE(IslandHandles.Num() == 1);
+	CHECK(Graph->NumIslands() == 1);
+
+	TArray<FGraphEdgeHandle> Edges = GetEdgesForVertex(VertexHandles[2]);
+	CHECK(Edges.Num() == 2);
+	Graph->RemoveEdge(Edges[0]);
+	Graph->RemoveEdge(Edges[1]);
+
+	UGraphIsland* Island1 = VertexHandles[0].GetVertex()->GetParentIsland().GetIsland();
+	REQUIRE(Island1 != nullptr);
+	{
+		CHECK(Island1->Num() == 2);
+		CHECK(Island1->GetVertices().Contains(VertexHandles[0]) == true);
+		CHECK(Island1->GetVertices().Contains(VertexHandles[1]) == true);
+	}
+
+	UGraphIsland* Island2 = VertexHandles[2].GetVertex()->GetParentIsland().GetIsland();
+	REQUIRE(Island2 != nullptr);
+	CHECK(Island1->Handle() != Island2->Handle());
+	{
+		CHECK(Island2->Num() == 1);
+		CHECK(Island2->GetVertices().Contains(VertexHandles[2]) == true);
+	}
+
+	UGraphIsland* Island3 = VertexHandles[3].GetVertex()->GetParentIsland().GetIsland();
+	REQUIRE(Island3 != nullptr);
+	CHECK(Island1->Handle() != Island3->Handle());
+	CHECK(Island2->Handle() != Island3->Handle());
+	{
+		CHECK(Island3->Num() == 2);
+		CHECK(Island3->GetVertices().Contains(VertexHandles[3]) == true);
+		CHECK(Island3->GetVertices().Contains(VertexHandles[4]) == true);
+	}
+
+	IslandVertexParentIslandSanityCheck(Island1->Handle());
+	IslandVertexParentIslandSanityCheck(Island2->Handle());
+	IslandVertexParentIslandSanityCheck(Island3->Handle());
+}
+
+TEST_CASE_METHOD(FTestGraphBuilder, "Graph::Island::Operation Allowed::Flag", "[graph][island]")
+{
+	PopulateVertices(2, true);
+	BuildLinearEdges(2);
+
+	REQUIRE(IslandHandles.Num() == 1);
+
+	UGraphIsland* Island = IslandHandles[0].GetIsland();
+	REQUIRE(Island != nullptr);
+
+	TArray<EGraphIslandOperations> AllFlags =
+	{
+		EGraphIslandOperations::Add,
+		EGraphIslandOperations::Split,
+		EGraphIslandOperations::Merge,
+		EGraphIslandOperations::Destroy
+	};
+
+	for (EGraphIslandOperations Flag : AllFlags)
+	{
+		CHECK(Island->IsOperationAllowed(Flag) == true);
+	}
+
+	auto RunTest = [&AllFlags, Island](EGraphIslandOperations FlagToDisable)
+	{
+		Island->SetOperationAllowed(FlagToDisable, false);
+		for (EGraphIslandOperations Flag : AllFlags)
+		{
+			if (Flag == FlagToDisable)
+			{
+				CHECK(Island->IsOperationAllowed(Flag) == false);
+			}
+			else
+			{
+				CHECK(Island->IsOperationAllowed(Flag) == true);
+			}
+		}
+
+		Island->SetOperationAllowed(FlagToDisable, true);
+		for (EGraphIslandOperations Flag : AllFlags)
+		{
+			CHECK(Island->IsOperationAllowed(Flag) == true);
+		}
+	};
+
+	RunTest(EGraphIslandOperations::Add);
+	RunTest(EGraphIslandOperations::Split);
+	RunTest(EGraphIslandOperations::Merge);
+	RunTest(EGraphIslandOperations::Destroy);
+}
+
+TEST_CASE_METHOD(FTestGraphBuilder, "Graph::Island::Operation Allowed::Disable Add", "[graph][island]")
+{
+	PopulateVertices(4, true);
+	BuildLinearEdges(4);
+	IslandVertexParentIslandSanityCheck(IslandHandles[0]);
+
+	REQUIRE(IslandHandles.Num() == 1);
+
+	UGraphIsland* Island = IslandHandles[0].GetIsland();
+	REQUIRE(Island != nullptr);
+	Island->SetOperationAllowed(EGraphIslandOperations::Add, false);
+
+	FGraphVertexHandle NewVertex = Graph->CreateVertex();
+
+	FEdgeCreationParameters Params;
+	Params.VertexHandle1 = NewVertex;
+	Params.VertexHandle2 = VertexHandles[0];
+
+	TArray<FGraphEdgeHandle> Edges;
+	Graph->CreateBulkEdges({ Params }, &Edges);
+
+	CHECK(Edges.Num() == 0);
+	CHECK(Island->Num() == 4);
+	CHECK(Island->GetVertices().Contains(NewVertex) == false);
+	CHECK(NewVertex.GetVertex()->GetParentIsland() != Island->Handle());
+	CHECK(NewVertex.GetVertex()->GetParentIsland() == FGraphIslandHandle::Invalid);
+	CHECK(Graph->NumVertices() == 5);
+	CHECK(Graph->NumEdges() == 3);
+	CHECK(Graph->NumIslands() == 1);
+	IslandVertexParentIslandSanityCheck(IslandHandles[0]);
+}
+
+TEST_CASE_METHOD(FTestGraphBuilder, "Graph::Island::Operation Allowed::Disable Split", "[graph][island]")
+{
+	PopulateVertices(6, true);
+	BuildLinearEdges(6);
+	IslandVertexParentIslandSanityCheck(IslandHandles[0]);
+
+	REQUIRE(IslandHandles.Num() == 1);
+
+	UGraphIsland* Island = IslandHandles[0].GetIsland();
+	REQUIRE(Island != nullptr);
+	Island->SetOperationAllowed(EGraphIslandOperations::Split, false);
+
+	CHECK(Island->Num() == 6);
+	CHECK(Graph->NumVertices() == 6);
+	CHECK(Graph->NumEdges() == 5);
+	CHECK(Graph->NumIslands() == 1);
+
+	UGraphVertex* Vertex = VertexHandles[2].GetVertex();
+	REQUIRE(Vertex != nullptr);
+
+	UGraphIsland* Island1 = nullptr;
+	UGraphIsland* Island2 = nullptr;
+
+	SECTION("Remove Edge")
+	{
+		const FGraphEdgeHandle* EdgeHandle = Vertex->FindEdgeTo(VertexHandles[3]);
+		REQUIRE(EdgeHandle != nullptr);
+		Graph->RemoveEdge(*EdgeHandle);
+
+		IslandVertexParentIslandSanityCheck(IslandHandles[0]);
+		CHECK(Island->GetVertices().Contains(VertexHandles[0]) == true);
+		CHECK(Island->GetVertices().Contains(VertexHandles[1]) == true);
+		CHECK(Island->GetVertices().Contains(VertexHandles[2]) == true);
+		CHECK(Island->GetVertices().Contains(VertexHandles[3]) == true);
+		CHECK(Island->GetVertices().Contains(VertexHandles[4]) == true);
+		CHECK(Island->GetVertices().Contains(VertexHandles[5]) == true);
+		CHECK(Island->Num() == 6);
+		CHECK(Graph->NumVertices() == 6);
+		CHECK(Graph->NumEdges() == 4);
+
+		CHECK(Graph->NumIslands() == 1);
+		Island->SetOperationAllowed(EGraphIslandOperations::Split, true);
+		Graph->RefreshIslandConnectivity(Island->Handle());
+
+		Island1 = VertexHandles[0].GetVertex()->GetParentIsland().GetIsland();
+		Island2 = VertexHandles[5].GetVertex()->GetParentIsland().GetIsland();
+
+		CHECK(Island1->Num() == 3);
+		CHECK(Island1->GetVertices().Contains(VertexHandles[0]) == true);
+		CHECK(Island1->GetVertices().Contains(VertexHandles[1]) == true);
+		CHECK(Island1->GetVertices().Contains(VertexHandles[2]) == true);
+		CHECK(Island2->Num() == 3);
+		CHECK(Island2->GetVertices().Contains(VertexHandles[3]) == true);
+		CHECK(Island2->GetVertices().Contains(VertexHandles[4]) == true);
+		CHECK(Island2->GetVertices().Contains(VertexHandles[5]) == true);
+		CHECK(Graph->NumVertices() == 6);
+		CHECK(Graph->NumEdges() == 4);
+	}
+
+	SECTION("Remove Vertex")
+	{
+		Graph->RemoveVertex(VertexHandles[2]);
+		IslandVertexParentIslandSanityCheck(IslandHandles[0]);
+		CHECK(Island->GetVertices().Contains(VertexHandles[0]) == true);
+		CHECK(Island->GetVertices().Contains(VertexHandles[1]) == true);
+		CHECK(Island->GetVertices().Contains(VertexHandles[3]) == true);
+		CHECK(Island->GetVertices().Contains(VertexHandles[4]) == true);
+		CHECK(Island->GetVertices().Contains(VertexHandles[5]) == true);
+		CHECK(Island->Num() == 5);
+		CHECK(Graph->NumVertices() == 5);
+		CHECK(Graph->NumEdges() == 3);
+
+		CHECK(Graph->NumIslands() == 1);
+		Island->SetOperationAllowed(EGraphIslandOperations::Split, true);
+		Graph->RefreshIslandConnectivity(Island->Handle());
+
+		Island1 = VertexHandles[0].GetVertex()->GetParentIsland().GetIsland();
+		Island2 = VertexHandles[5].GetVertex()->GetParentIsland().GetIsland();
+
+		CHECK(Island1->Num() == 2);
+		CHECK(Island1->GetVertices().Contains(VertexHandles[0]) == true);
+		CHECK(Island1->GetVertices().Contains(VertexHandles[1]) == true);
+		CHECK(Island2->Num() == 3);
+		CHECK(Island2->GetVertices().Contains(VertexHandles[3]) == true);
+		CHECK(Island2->GetVertices().Contains(VertexHandles[4]) == true);
+		CHECK(Island2->GetVertices().Contains(VertexHandles[5]) == true);
+		CHECK(Graph->NumVertices() == 5);
+		CHECK(Graph->NumEdges() == 3);
+	}
+
+	CHECK(Graph->NumIslands() == 2);
+	IslandVertexParentIslandSanityCheck(Island1->Handle());
+	IslandVertexParentIslandSanityCheck(Island2->Handle());
+}
+
+TEST_CASE_METHOD(FTestGraphBuilder, "Graph::Island::Operation Allowed::Disable Merge::Existing Vertices", "[graph][island]")
+{
+	PopulateVertices(4, true);
+	BuildLinearEdges(2);
+
+	REQUIRE(IslandHandles.Num() == 2);
+	IslandVertexParentIslandSanityCheck(IslandHandles[0]);
+	IslandVertexParentIslandSanityCheck(IslandHandles[1]);
+
+	auto RunTest = [this]()
+	{
+		FEdgeCreationParameters Params;
+		Params.VertexHandle1 = VertexHandles[0];
+		Params.VertexHandle2 = VertexHandles[2];
+
+		TArray<FGraphEdgeHandle> Edges;
+		Graph->CreateBulkEdges({ Params }, &Edges);
+
+		CHECK(Edges.Num() == 0);
+		CHECK(Graph->NumVertices() == 4);
+		CHECK(Graph->NumEdges() == 2);
+		CHECK(Graph->NumIslands() == 2);
+
+		UGraphIsland* Island0 = VertexHandles[0].GetVertex()->GetParentIsland().GetIsland();
+		REQUIRE(Island0 != nullptr);
+		CHECK(Island0->Num() == 2);
+		CHECK(Island0->GetVertices().Contains(VertexHandles[0]) == true);
+		CHECK(Island0->GetVertices().Contains(VertexHandles[1]) == true);
+
+		UGraphIsland* Island1 = VertexHandles[2].GetVertex()->GetParentIsland().GetIsland();
+		REQUIRE(Island1 != nullptr);
+		CHECK(Island1->Num() == 2);
+		CHECK(Island0->Handle() != Island1->Handle());
+		CHECK(Island1->GetVertices().Contains(VertexHandles[2]) == true);
+		CHECK(Island1->GetVertices().Contains(VertexHandles[3]) == true);
+
+
+		IslandVertexParentIslandSanityCheck(Island0->Handle());
+		IslandVertexParentIslandSanityCheck(Island1->Handle());
+	};
+
+	SECTION("Single disable - 1")
+	{
+		UGraphIsland* Island = IslandHandles[0].GetIsland();
+		REQUIRE(Island != nullptr);
+		Island->SetOperationAllowed(EGraphIslandOperations::Merge, false);
+
+		RunTest();
+	}
+
+	SECTION("Single disable - 2")
+	{
+		UGraphIsland* Island = IslandHandles[1].GetIsland();
+		REQUIRE(Island != nullptr);
+		Island->SetOperationAllowed(EGraphIslandOperations::Merge, false);
+
+		RunTest();
+	}
+
+	SECTION("Double disable")
+	{
+		UGraphIsland* Island0 = IslandHandles[0].GetIsland();
+		REQUIRE(Island0 != nullptr);
+		Island0->SetOperationAllowed(EGraphIslandOperations::Merge, false);
+
+		UGraphIsland* Island1 = IslandHandles[1].GetIsland();
+		REQUIRE(Island1 != nullptr);
+		Island1->SetOperationAllowed(EGraphIslandOperations::Merge, false);
+
+		RunTest();
+	}
+}
+
+TEST_CASE_METHOD(FTestGraphBuilder, "Graph::Island::Operation Allowed::Disable Merge::New Vertex::Bulk Edges", "[graph][island]")
+{
+	PopulateVertices(4, true);
+	BuildLinearEdges(2);
+
+	REQUIRE(IslandHandles.Num() == 2);
+	IslandVertexParentIslandSanityCheck(IslandHandles[0]);
+	IslandVertexParentIslandSanityCheck(IslandHandles[1]);
+
+	auto RunTest = [this]()
+	{
+		FGraphVertexHandle NewVertex = Graph->CreateVertex();
+		FEdgeCreationParameters Params1;
+		Params1.VertexHandle1 = VertexHandles[0];
+		Params1.VertexHandle2 = NewVertex;
+
+		FEdgeCreationParameters Params2;
+		Params2.VertexHandle1 = VertexHandles[3];
+		Params2.VertexHandle2 = NewVertex;
+
+		TArray<FGraphEdgeHandle> Edges;
+		Graph->CreateBulkEdges({ Params1, Params2 }, &Edges);
+		Graph->FinalizeVertex(NewVertex);
+
+		CHECK(Edges.Num() == 1);
+		CHECK(Graph->NumVertices() == 5);
+		CHECK(Graph->NumEdges() == 3);
+		CHECK(Graph->NumIslands() == 2);
+
+		UGraphIsland* Island0 = VertexHandles[0].GetVertex()->GetParentIsland().GetIsland();
+		REQUIRE(Island0 != nullptr);
+		CHECK(Island0->GetVertices().Contains(VertexHandles[0]) == true);
+		CHECK(Island0->GetVertices().Contains(VertexHandles[1]) == true);
+
+		UGraphIsland* Island1 = VertexHandles[2].GetVertex()->GetParentIsland().GetIsland();
+		REQUIRE(Island1 != nullptr);
+		CHECK(Island0->Handle() != Island1->Handle());
+		CHECK(Island1->GetVertices().Contains(VertexHandles[2]) == true);
+		CHECK(Island1->GetVertices().Contains(VertexHandles[3]) == true);
+
+		CHECK(Island0->Num() == 3);
+		CHECK(Island0->GetVertices().Contains(NewVertex) == true);
+		CHECK(Island1->Num() == 2);
+
+		IslandVertexParentIslandSanityCheck(Island0->Handle());
+		IslandVertexParentIslandSanityCheck(Island1->Handle());
+	};
+
+	SECTION("Single disable - 1")
+	{
+		UGraphIsland* Island = VertexHandles[0].GetVertex()->GetParentIsland().GetIsland();
+		REQUIRE(Island != nullptr);
+		Island->SetOperationAllowed(EGraphIslandOperations::Merge, false);
+
+		RunTest();
+	}
+
+	SECTION("Single disable - 2")
+	{
+		UGraphIsland* Island = VertexHandles[3].GetVertex()->GetParentIsland().GetIsland();
+		REQUIRE(Island != nullptr);
+		Island->SetOperationAllowed(EGraphIslandOperations::Merge, false);
+
+		RunTest();
+	}
+
+	SECTION("Double disable")
+	{
+		UGraphIsland* Island0 = VertexHandles[0].GetVertex()->GetParentIsland().GetIsland();
+		REQUIRE(Island0 != nullptr);
+		Island0->SetOperationAllowed(EGraphIslandOperations::Merge, false);
+
+		UGraphIsland* Island1 = VertexHandles[3].GetVertex()->GetParentIsland().GetIsland();
+		REQUIRE(Island1 != nullptr);
+		Island1->SetOperationAllowed(EGraphIslandOperations::Merge, false);
+
+		RunTest();
+	}
+}
+
+TEST_CASE_METHOD(FTestGraphBuilder, "Graph::Island::Operation Allowed::Disable Destroy", "[graph][island]")
+{
+	PopulateVertices(4, true);
+	BuildFullyConnectedEdges(4);
+
+	REQUIRE(IslandHandles.Num() == 1);
+	CHECK(Graph->NumIslands() == 1);
+	IslandVertexParentIslandSanityCheck(IslandHandles[0]);
+	CHECK(IslandHandles[0].IsComplete());
+
+	UGraphIsland* Island = IslandHandles[0].GetIsland();
+	REQUIRE(Island != nullptr);
+	Island->SetOperationAllowed(EGraphIslandOperations::Destroy, false);
+
+	Graph->RemoveIsland(IslandHandles[0]);
+
+	CHECK(Graph->NumIslands() == 1);
+	IslandVertexParentIslandSanityCheck(IslandHandles[0]);
+	CHECK(IslandHandles[0].IsComplete());
+	CHECK(Graph->GetIslands().Contains(IslandHandles[0]) == true);
+	CHECK(Island->Num() == 4);
+	CHECK(Island->GetVertices().Contains(VertexHandles[0]) == true);
+	CHECK(Island->GetVertices().Contains(VertexHandles[1]) == true);
+	CHECK(Island->GetVertices().Contains(VertexHandles[2]) == true);
+	CHECK(Island->GetVertices().Contains(VertexHandles[3]) == true);
+}
+
+TEST_CASE_METHOD(FTestGraphBuilder, "Graph::Island::Get Serialized Data", "[graph][island]")
+{
+	PopulateVertices(4, true);
+	BuildLinearEdges(4);
+
+	REQUIRE(IslandHandles.Num() == 1);
+
+	UGraphIsland* Island = IslandHandles[0].GetIsland();
+	REQUIRE(Island != nullptr);
+
+	FSerializedIslandData Serialized = Island->GetSerializedData();
+	CHECK(Island->Num() == Serialized.Vertices.Num());
+	TSet<FGraphVertexHandle> SerializedVertexSet{ Serialized.Vertices };
+	CHECK(Island->GetVertices().Num() == SerializedVertexSet.Num());
+	for (const FGraphVertexHandle& Handle : SerializedVertexSet)
+	{
+		CHECK(Island->GetVertices().Contains(Handle));
+	}
+}
+
+TEST_CASE_METHOD(FTestGraphBuilder, "Graph::Island::Events::Vertex Added", "[graph][island]")
+{
+	PopulateVertices(2, true);
+	BuildLinearEdges(2);
+
+	REQUIRE(IslandHandles.Num() == 1);
+
+	UGraphIsland* Island = IslandHandles[0].GetIsland();
+	REQUIRE(Island != nullptr);
+	IslandVertexParentIslandSanityCheck(IslandHandles[0]);
+
+	FGraphIslandHandle CallbackIslandHandle;
+	FGraphVertexHandle CallbackVertexHandle;
+	Island->OnVertexAdded.AddLambda(
+		[&CallbackIslandHandle, &CallbackVertexHandle](const FGraphIslandHandle& InIslandHandle, const FGraphVertexHandle& InVertexHandle)
+		{
+			CallbackIslandHandle = InIslandHandle;
+			CallbackVertexHandle = InVertexHandle;
+		}
+	);
+
+	FGraphVertexHandle NewVertex = Graph->CreateVertex();
+	FEdgeCreationParameters Params;
+	Params.VertexHandle1 = NewVertex;
+	Params.VertexHandle2 = VertexHandles[0];
+	Graph->CreateBulkEdges({ Params });
+
+	CHECK(CallbackIslandHandle == Island->Handle());
+	CHECK(CallbackVertexHandle == NewVertex);
+}
+
+TEST_CASE_METHOD(FTestGraphBuilder, "Graph::Island::Events::Vertex Removed", "[graph][island]")
+{
+	PopulateVertices(3, true);
+	BuildLinearEdges(3);
+
+	REQUIRE(IslandHandles.Num() == 1);
+
+	UGraphIsland* Island = IslandHandles[0].GetIsland();
+	REQUIRE(Island != nullptr);
+	IslandVertexParentIslandSanityCheck(IslandHandles[0]);
+
+	FGraphIslandHandle CallbackIslandHandle;
+	FGraphVertexHandle CallbackVertexHandle;
+	Island->OnVertexRemoved.AddLambda(
+		[&CallbackIslandHandle, &CallbackVertexHandle](const FGraphIslandHandle& InIslandHandle, const FGraphVertexHandle& InVertexHandle)
+		{
+			CallbackIslandHandle = InIslandHandle;
+			CallbackVertexHandle = InVertexHandle;
+		}
+	);
+
+	SECTION("Remove Vertex")
+	{
+		Graph->RemoveVertex(VertexHandles[0]);
+		CHECK(CallbackIslandHandle == Island->Handle());
+		CHECK(CallbackVertexHandle == VertexHandles[0]);
+	}
+
+	SECTION("Remove Edge")
+	{
+		TArray<FGraphEdgeHandle> Edges = GetEdgesForVertex(VertexHandles[0]);
+		CHECK(Edges.Num() == 1);
+		Graph->RemoveEdge(Edges[0]);
+		CHECK(CallbackIslandHandle == Island->Handle());
+		CHECK(CallbackVertexHandle == VertexHandles[0]);
+	}
+}
+
+TEST_CASE_METHOD(FTestGraphBuilder, "Graph::Island::Events::On Destroyed", "[graph][island]")
+{
+	PopulateVertices(3, true);
+	BuildLinearEdges(3);
+
+	REQUIRE(IslandHandles.Num() == 1);
+
+	UGraphIsland* Island = IslandHandles[0].GetIsland();
+	REQUIRE(Island != nullptr);
+	IslandVertexParentIslandSanityCheck(IslandHandles[0]);
+
+	FGraphIslandHandle CallbackIslandHandle;
+	Island->OnDestroyed.AddLambda(
+		[&CallbackIslandHandle](const FGraphIslandHandle& InIslandHandle)
+		{
+			CallbackIslandHandle = InIslandHandle;
+		}
+	);
+
+	Graph->RemoveIsland(IslandHandles[0]);
+	CHECK(CallbackIslandHandle == IslandHandles[0]);
+}
+
+TEST_CASE_METHOD(FTestGraphBuilder, "Graph::Island::Events::On Connectivity Changed", "[graph][island]")
+{
+	PopulateVertices(3, false);
+
+	SECTION("Finalize Vertex")
+	{
+		FGraphIslandHandle CreatedIslandHandle;
+		FGraphIslandHandle CallbackIslandHandle;
+		Graph->OnIslandCreated.AddLambda(
+			[&CreatedIslandHandle, &CallbackIslandHandle](const FGraphIslandHandle& InIslandHandle)
+			{
+				CreatedIslandHandle = InIslandHandle;
+
+				UGraphIsland* Island = CreatedIslandHandle.GetIsland();
+				REQUIRE(Island != nullptr);
+
+				Island->OnConnectivityChanged.AddLambda(
+					[&CallbackIslandHandle](const FGraphIslandHandle& InCallbackIslandHandle)
+					{
+						CallbackIslandHandle = InCallbackIslandHandle;
+					}
+				);
+			}
+		);
+
+		Graph->FinalizeVertex(VertexHandles[0]);
+		CHECK(CreatedIslandHandle.IsComplete() == true);
+		CHECK(CallbackIslandHandle.IsComplete() == true);
+		CHECK(CreatedIslandHandle == CallbackIslandHandle);
+	}
+
+	SECTION("Split island")
+	{
+		FinalizeVertices();
+		BuildLinearEdges(3);
+
+		REQUIRE(IslandHandles.Num() == 1);
+
+		UGraphIsland* Island = IslandHandles[0].GetIsland();
+		REQUIRE(Island != nullptr);
+		IslandVertexParentIslandSanityCheck(IslandHandles[0]);
+
+		FGraphIslandHandle CallbackIslandHandle;
+		Island->OnConnectivityChanged.AddLambda(
+			[&CallbackIslandHandle](const FGraphIslandHandle& InCallbackIslandHandle)
+			{
+				CallbackIslandHandle = InCallbackIslandHandle;
+			}
+		);
+
+		Graph->RemoveVertex(VertexHandles[1]);
+
+		CHECK(CallbackIslandHandle.IsComplete());
+		CHECK(CallbackIslandHandle == IslandHandles[0]);
+	}
+}
\ No newline at end of file
diff --git a/Engine/Plugins/Experimental/GameplayGraph/Tests/Private/GraphSerializationTests.cpp b/Engine/Plugins/Experimental/GameplayGraph/Tests/Private/GraphSerializationTests.cpp
new file mode 100644
--- /dev/null
+++ b/Engine/Plugins/Experimental/GameplayGraph/Tests/Private/GraphSerializationTests.cpp
@@ -0,0 +1,279 @@
+// Copyright Epic Games, Inc. All Rights Reserved.
+#include "TestHarness.h"
+#include "TestGraphBuilder.h"
+
+TEST_CASE_METHOD(FTestGraphBuilder, "Graph::Serialization::Write", "[graph]")
+{
+	PopulateVertices(6, true);
+	BuildLinearEdges(3);
+
+	CHECK(Graph->NumVertices() == 6);
+	CHECK(Graph->NumEdges() == 4);
+	CHECK(Graph->NumIslands() == 2);
+
+	FSerializableGraph SerializedGraph = Graph->GetSerializableGraph();
+	CHECK(SerializedGraph.Properties == Graph->GetProperties());
+	CHECK(SerializedGraph.Vertices.Num() == 6);
+	CHECK(SerializedGraph.Vertices.Contains(VertexHandles[0]) == true);
+	CHECK(SerializedGraph.Vertices.Contains(VertexHandles[1]) == true);
+	CHECK(SerializedGraph.Vertices.Contains(VertexHandles[2]) == true);
+	CHECK(SerializedGraph.Vertices.Contains(VertexHandles[3]) == true);
+	CHECK(SerializedGraph.Vertices.Contains(VertexHandles[4]) == true);
+	CHECK(SerializedGraph.Vertices.Contains(VertexHandles[5]) == true);
+
+	CHECK(SerializedGraph.Edges.Num() == 4);
+	for (const FGraphVertexHandle& VertexHandle : VertexHandles)
+	{
+		UGraphVertex* Vertex = VertexHandle.GetVertex();
+		REQUIRE(Vertex != nullptr);
+
+		Vertex->ForEachAdjacentVertex(
+			[&SerializedGraph](const FGraphVertexHandle& NeighborVertexHandle, const FGraphEdgeHandle& EdgeHandle)
+			{
+				REQUIRE(SerializedGraph.Edges.Contains(EdgeHandle) == true);
+				const FSerializedEdgeData& EdgeData = SerializedGraph.Edges[EdgeHandle];
+				UGraphEdge* Edge = EdgeHandle.GetEdge();
+				REQUIRE(Edge != nullptr);
+				CHECK(EdgeData.Node1 == Edge->NodeA());
+				CHECK(EdgeData.Node2 == Edge->NodeB());
+			}
+		);
+	}
+
+	CHECK(SerializedGraph.Islands.Num() == 2);
+	CHECK(SerializedGraph.Islands.Contains(IslandHandles[0]) == true);
+	CHECK(SerializedGraph.Islands.Contains(IslandHandles[1]) == true);
+	for (const FGraphIslandHandle& IslandHandle : IslandHandles)
+	{
+		UGraphIsland* Island = IslandHandle.GetIsland();
+		REQUIRE(Island != nullptr);
+
+		REQUIRE(SerializedGraph.Islands.Contains(IslandHandle) == true);
+		const FSerializedIslandData& IslandData = SerializedGraph.Islands[IslandHandle];
+
+		CHECK(IslandData.Vertices.Num() == Island->Num());
+		for (const FGraphVertexHandle& VertexHandle : Island->GetVertices())
+		{
+			CHECK(IslandData.Vertices.Contains(VertexHandle) == true);
+		}
+	}
+}
+
+TEST_CASE_METHOD(FTestGraphBuilder, "Graph::Serialization::Read (normal)", "[graph]")
+{
+	FSerializableGraph SerializedGraph;
+	SerializedGraph.Properties.bGenerateIslands = true;
+
+	SerializedGraph.Vertices.Add(FGraphVertexHandle { FGraphUniqueIndex::CreateUniqueIndex(), nullptr });
+	SerializedGraph.Vertices.Add(FGraphVertexHandle { FGraphUniqueIndex::CreateUniqueIndex(), nullptr });
+	SerializedGraph.Vertices.Add(FGraphVertexHandle { FGraphUniqueIndex::CreateUniqueIndex(), nullptr });
+	SerializedGraph.Vertices.Add(FGraphVertexHandle { FGraphUniqueIndex::CreateUniqueIndex(), nullptr });
+	SerializedGraph.Vertices.Add(FGraphVertexHandle { FGraphUniqueIndex::CreateUniqueIndex(), nullptr });
+	SerializedGraph.Vertices.Add(FGraphVertexHandle { FGraphUniqueIndex::CreateUniqueIndex(), nullptr });
+
+	SerializedGraph.Edges.Add(FGraphEdgeHandle { FGraphUniqueIndex::CreateUniqueIndex(), nullptr }, FSerializedEdgeData{ SerializedGraph.Vertices[0], SerializedGraph.Vertices[1] });
+	SerializedGraph.Edges.Add(FGraphEdgeHandle { FGraphUniqueIndex::CreateUniqueIndex(), nullptr }, FSerializedEdgeData{ SerializedGraph.Vertices[1], SerializedGraph.Vertices[2] });
+	SerializedGraph.Edges.Add(FGraphEdgeHandle { FGraphUniqueIndex::CreateUniqueIndex(), nullptr }, FSerializedEdgeData{ SerializedGraph.Vertices[3], SerializedGraph.Vertices[4] });
+	SerializedGraph.Edges.Add(FGraphEdgeHandle { FGraphUniqueIndex::CreateUniqueIndex(), nullptr }, FSerializedEdgeData{ SerializedGraph.Vertices[4], SerializedGraph.Vertices[5] });
+
+	SerializedGraph.Islands.Add(FGraphIslandHandle { FGraphUniqueIndex::CreateUniqueIndex(), nullptr }, FSerializedIslandData { TArray{SerializedGraph.Vertices[0], SerializedGraph.Vertices[1], SerializedGraph.Vertices[2]} });
+	SerializedGraph.Islands.Add(FGraphIslandHandle { FGraphUniqueIndex::CreateUniqueIndex(), nullptr }, FSerializedIslandData { TArray{SerializedGraph.Vertices[3], SerializedGraph.Vertices[4], SerializedGraph.Vertices[5]} });
+	Graph->LoadFromSerializedGraph(SerializedGraph);
+
+	CHECK(SerializedGraph.Properties == Graph->GetProperties());
+	CHECK(Graph->NumVertices() == 6);
+	CHECK(Graph->NumEdges() == 4);
+	CHECK(Graph->NumIslands() == 2);
+
+	int32 Index = 0;
+	for (const FGraphVertexHandle& SerializedVertexHandle : SerializedGraph.Vertices)
+	{
+		CHECK(Graph->GetVertices().Contains(SerializedVertexHandle) == true);
+
+		UGraphVertex* LoadedVertex = Graph->GetVertices().FindRef(SerializedVertexHandle);
+		REQUIRE(LoadedVertex != nullptr);
+		CHECK(SerializedVertexHandle == LoadedVertex->Handle());
+	}
+
+	for (const TPair<FGraphEdgeHandle, FSerializedEdgeData>& Edge: SerializedGraph.Edges)
+	{
+		CHECK(Graph->GetEdges().Contains(Edge.Key) == true);
+
+		CHECK(Graph->GetVertices().Contains(Edge.Value.Node1) == true);
+		CHECK(Graph->GetVertices().Contains(Edge.Value.Node2) == true);
+
+		UGraphEdge* LoadedEdge = Graph->GetEdges().FindRef(Edge.Key);
+		REQUIRE(LoadedEdge != nullptr);
+		CHECK(LoadedEdge->ContainsNode(Edge.Value.Node1) == true);
+		CHECK(LoadedEdge->NodeA().IsComplete() == true);
+		CHECK(LoadedEdge->ContainsNode(Edge.Value.Node2) == true);
+		CHECK(LoadedEdge->NodeB().IsComplete() == true);
+	}
+
+	for (const TPair<FGraphIslandHandle, FSerializedIslandData>& Island : SerializedGraph.Islands)
+	{
+		CHECK(Graph->GetIslands().Contains(Island.Key) == true);
+
+		UGraphIsland* LoadedIsland = Graph->GetIslands().FindRef(Island.Key);
+		REQUIRE(LoadedIsland != nullptr);
+
+		for (const FGraphVertexHandle& IslandVertexHandle : Island.Value.Vertices)
+		{
+			CHECK(LoadedIsland->GetVertices().Contains(IslandVertexHandle) == true);
+
+			const FGraphVertexHandle* LoadedIslandVertexHandle = LoadedIsland->GetVertices().Find(IslandVertexHandle);
+			REQUIRE(LoadedIslandVertexHandle != nullptr);
+			CHECK(LoadedIslandVertexHandle->IsComplete() == true);
+			CHECK(LoadedIslandVertexHandle->GetGraph() == Graph);
+		}
+	}
+}
+
+TEST_CASE_METHOD(FTestGraphBuilder, "Graph::Serialization::Read::Errors::Vertex::Invalid Index", "[graph]")
+{
+	FSerializableGraph SerializedGraph;
+	SerializedGraph.Properties.bGenerateIslands = true;
+	SerializedGraph.Vertices.Add(FGraphVertexHandle { FGraphUniqueIndex{ FGuid{0, 0, 0, 0} }, nullptr });
+
+	Graph->LoadFromSerializedGraph(SerializedGraph);
+	CHECK(SerializedGraph.Properties == Graph->GetProperties());
+	CHECK(Graph->NumVertices() == 0);
+	CHECK(Graph->NumEdges() == 0);
+	CHECK(Graph->NumIslands() == 0);
+}
+
+TEST_CASE_METHOD(FTestGraphBuilder, "Graph::Serialization::Read::Errors::Edge::Invalid Index", "[graph]")
+{
+	FSerializableGraph SerializedGraph;
+	SerializedGraph.Properties.bGenerateIslands = true;
+	SerializedGraph.Vertices.Add(FGraphVertexHandle { FGraphUniqueIndex::CreateUniqueIndex(), nullptr });
+	SerializedGraph.Vertices.Add(FGraphVertexHandle { FGraphUniqueIndex::CreateUniqueIndex(), nullptr });
+	SerializedGraph.Edges.Add(FGraphEdgeHandle { FGraphUniqueIndex{FGuid{0, 0, 0, 0}}, nullptr }, FSerializedEdgeData{ SerializedGraph.Vertices[0], SerializedGraph.Vertices[1] });
+
+	Graph->LoadFromSerializedGraph(SerializedGraph);
+	CHECK(SerializedGraph.Properties == Graph->GetProperties());
+	CHECK(Graph->NumVertices() == 2);
+	CHECK(Graph->NumEdges() == 0);
+	CHECK(Graph->NumIslands() == 0);
+}
+
+TEST_CASE_METHOD(FTestGraphBuilder, "Graph::Serialization::Read::Errors::Edge::Invalid Vertex Index", "[graph]")
+{
+	FSerializableGraph SerializedGraph;
+	SerializedGraph.Properties.bGenerateIslands = true;
+	SerializedGraph.Vertices.Add(FGraphVertexHandle { FGraphUniqueIndex{FGuid{0, 0, 0, 1}}, nullptr });
+	SerializedGraph.Vertices.Add(FGraphVertexHandle { FGraphUniqueIndex{FGuid{0, 0, 0, 2}}, nullptr });
+	SerializedGraph.Edges.Add(FGraphEdgeHandle { FGraphUniqueIndex::CreateUniqueIndex(), nullptr }, FSerializedEdgeData{ SerializedGraph.Vertices[0], FGraphVertexHandle { FGraphUniqueIndex{FGuid{0, 0, 0, 3}}, nullptr } });
+
+	Graph->LoadFromSerializedGraph(SerializedGraph);
+	CHECK(SerializedGraph.Properties == Graph->GetProperties());
+	CHECK(Graph->NumVertices() == 2);
+	CHECK(Graph->NumEdges() == 0);
+	CHECK(Graph->NumIslands() == 0);
+}
+
+TEST_CASE_METHOD(FTestGraphBuilder, "Graph::Serialization::Read::Errors::Island::Invalid Index", "[graph]")
+{
+	FSerializableGraph SerializedGraph;
+	SerializedGraph.Properties.bGenerateIslands = true;
+	SerializedGraph.Vertices.Add(FGraphVertexHandle { FGraphUniqueIndex::CreateUniqueIndex(), nullptr });
+	SerializedGraph.Vertices.Add(FGraphVertexHandle { FGraphUniqueIndex::CreateUniqueIndex(), nullptr });
+	SerializedGraph.Edges.Add(FGraphEdgeHandle { FGraphUniqueIndex::CreateUniqueIndex(), nullptr }, FSerializedEdgeData{ SerializedGraph.Vertices[0], SerializedGraph.Vertices[1] });
+	SerializedGraph.Islands.Add(FGraphIslandHandle { FGraphUniqueIndex{FGuid{0, 0, 0, 0}}, nullptr }, FSerializedIslandData { TArray{SerializedGraph.Vertices[0], SerializedGraph.Vertices[1] } });
+
+	Graph->LoadFromSerializedGraph(SerializedGraph);
+	CHECK(SerializedGraph.Properties == Graph->GetProperties());
+	CHECK(Graph->NumVertices() == 2);
+	CHECK(Graph->NumEdges() == 1);
+	CHECK(Graph->NumIslands() == 1);
+	CHECK(Graph->GetIslands().Contains(FGraphIslandHandle { FGraphUniqueIndex { FGuid { 0, 0, 0, 0 } }, nullptr }) == false);
+}
+
+TEST_CASE_METHOD(FTestGraphBuilder, "Graph::Serialization::Read::Errors::Island::Invalid Vertex Index", "[graph]")
+{
+	FSerializableGraph SerializedGraph;
+	SerializedGraph.Properties.bGenerateIslands = true;
+	SerializedGraph.Vertices.Add(FGraphVertexHandle { FGraphUniqueIndex{FGuid{0, 0, 0, 1}}, nullptr });
+	SerializedGraph.Vertices.Add(FGraphVertexHandle { FGraphUniqueIndex{FGuid{0, 0, 0, 2}}, nullptr });
+	SerializedGraph.Edges.Add(FGraphEdgeHandle { FGraphUniqueIndex::CreateUniqueIndex(), nullptr }, FSerializedEdgeData{ SerializedGraph.Vertices[0], SerializedGraph.Vertices[1] });
+	SerializedGraph.Islands.Add(FGraphIslandHandle { FGraphUniqueIndex{FGuid{1, 0, 0, 0}}, nullptr }, FSerializedIslandData { TArray{ SerializedGraph.Vertices[0], SerializedGraph.Vertices[1], FGraphVertexHandle { FGraphUniqueIndex{FGuid{0, 0, 0, 3}}, nullptr } } });
+
+	Graph->LoadFromSerializedGraph(SerializedGraph);
+	CHECK(SerializedGraph.Properties == Graph->GetProperties());
+	CHECK(Graph->NumVertices() == 2);
+	CHECK(Graph->NumEdges() == 1);
+	CHECK(Graph->NumIslands() == 1);
+
+	UGraphIsland* LoadedIsland = Graph->GetIslands().FindRef(FGraphIslandHandle {FGraphUniqueIndex{FGuid{1, 0, 0, 0}}, nullptr });
+	REQUIRE(LoadedIsland != nullptr);
+	CHECK(LoadedIsland->Num() == 2);
+	CHECK(LoadedIsland->GetVertices().Contains(SerializedGraph.Vertices[0]) == true);
+	CHECK(LoadedIsland->GetVertices().Contains(SerializedGraph.Vertices[1]) == true);
+	CHECK(LoadedIsland->GetVertices().Contains(FGraphVertexHandle { FGraphUniqueIndex { FGuid { 0, 0, 0, 3 } }, nullptr }) == false);
+}
+
+TEST_CASE_METHOD(FTestGraphBuilder, "Graph::Serialization::Read::Errors::Island::Edge Island Mismatch::Merge", "[graph]")
+{
+	FSerializableGraph SerializedGraph;
+	SerializedGraph.Properties.bGenerateIslands = true;
+	SerializedGraph.Vertices.Add(FGraphVertexHandle { FGraphUniqueIndex::CreateUniqueIndex(), nullptr });
+	SerializedGraph.Vertices.Add(FGraphVertexHandle { FGraphUniqueIndex::CreateUniqueIndex(), nullptr });
+	SerializedGraph.Vertices.Add(FGraphVertexHandle { FGraphUniqueIndex::CreateUniqueIndex(), nullptr });
+
+	SerializedGraph.Edges.Add(FGraphEdgeHandle { FGraphUniqueIndex::CreateUniqueIndex(), nullptr }, FSerializedEdgeData{ SerializedGraph.Vertices[0], SerializedGraph.Vertices[1] });
+	SerializedGraph.Edges.Add(FGraphEdgeHandle { FGraphUniqueIndex::CreateUniqueIndex(), nullptr }, FSerializedEdgeData{ SerializedGraph.Vertices[1], SerializedGraph.Vertices[2] });
+
+	SerializedGraph.Islands.Add(FGraphIslandHandle { FGraphUniqueIndex{FGuid{1, 0, 0, 0}}, nullptr }, FSerializedIslandData { TArray { SerializedGraph.Vertices[0], SerializedGraph.Vertices[1] } });
+	SerializedGraph.Islands.Add(FGraphIslandHandle { FGraphUniqueIndex{FGuid{2, 0, 0, 0}}, nullptr }, FSerializedIslandData { TArray { SerializedGraph.Vertices[2] } });
+
+	Graph->LoadFromSerializedGraph(SerializedGraph);
+	CHECK(SerializedGraph.Properties == Graph->GetProperties());
+	CHECK(Graph->NumVertices() == 3);
+	CHECK(Graph->NumEdges() == 2);
+	CHECK(Graph->NumIslands() == 1);
+
+	CHECK(Graph->GetIslands().Contains(FGraphIslandHandle { FGraphUniqueIndex { FGuid { 1, 0, 0, 0 } }, nullptr }) == true);
+	CHECK(Graph->GetIslands().Contains(FGraphIslandHandle { FGraphUniqueIndex { FGuid { 2, 0, 0, 0 } }, nullptr }) == false);
+
+	UGraphIsland* LoadedIsland = Graph->GetIslands().FindRef(FGraphIslandHandle {FGraphUniqueIndex{FGuid{1, 0, 0, 0}}, nullptr });
+	REQUIRE(LoadedIsland != nullptr);
+	CHECK(LoadedIsland->Num() == 3);
+	CHECK(LoadedIsland->GetVertices().Contains(SerializedGraph.Vertices[0]) == true);
+	CHECK(LoadedIsland->GetVertices().Contains(SerializedGraph.Vertices[1]) == true);
+	CHECK(LoadedIsland->GetVertices().Contains(SerializedGraph.Vertices[2]) == true);
+}
+
+TEST_CASE_METHOD(FTestGraphBuilder, "Graph::Serialization::Read::Errors::Island::Edge Island Mismatch::Split", "[graph]")
+{
+	FSerializableGraph SerializedGraph;
+	SerializedGraph.Properties.bGenerateIslands = true;
+	SerializedGraph.Vertices.Add(FGraphVertexHandle { FGraphUniqueIndex::CreateUniqueIndex(), nullptr });
+	SerializedGraph.Vertices.Add(FGraphVertexHandle { FGraphUniqueIndex::CreateUniqueIndex(), nullptr });
+	SerializedGraph.Vertices.Add(FGraphVertexHandle { FGraphUniqueIndex::CreateUniqueIndex(), nullptr });
+
+	SerializedGraph.Edges.Add(FGraphEdgeHandle { FGraphUniqueIndex::CreateUniqueIndex(), nullptr }, FSerializedEdgeData{ SerializedGraph.Vertices[0], SerializedGraph.Vertices[1] });
+
+	SerializedGraph.Islands.Add(FGraphIslandHandle { FGraphUniqueIndex{FGuid{1, 0, 0, 0}}, nullptr }, FSerializedIslandData { TArray { SerializedGraph.Vertices[0], SerializedGraph.Vertices[1], SerializedGraph.Vertices[2] } });
+
+	Graph->LoadFromSerializedGraph(SerializedGraph);
+	CHECK(SerializedGraph.Properties == Graph->GetProperties());
+	CHECK(Graph->NumVertices() == 3);
+	CHECK(Graph->NumEdges() == 1);
+	CHECK(Graph->NumIslands() == 2);
+
+	CHECK(Graph->GetIslands().Contains(FGraphIslandHandle { FGraphUniqueIndex { FGuid { 1, 0, 0, 0 } }, nullptr }) == true);
+
+	UGraphIsland* LoadedIsland = Graph->GetIslands().FindRef(FGraphIslandHandle {FGraphUniqueIndex{FGuid{1, 0, 0, 0}}, nullptr });
+	REQUIRE(LoadedIsland != nullptr);
+	CHECK(LoadedIsland->Num() == 2);
+	CHECK(LoadedIsland->GetVertices().Contains(SerializedGraph.Vertices[0]) == true);
+	CHECK(LoadedIsland->GetVertices().Contains(SerializedGraph.Vertices[1]) == true);
+	CHECK(LoadedIsland->GetVertices().Contains(SerializedGraph.Vertices[2]) == false);
+
+	UGraphIsland* SplitIsland = Graph->GetCompleteNodeHandle(SerializedGraph.Vertices[2]).GetVertex()->GetParentIsland().GetIsland();
+	REQUIRE(SplitIsland != nullptr);
+	CHECK(SplitIsland->Num() == 1);
+	CHECK(SplitIsland->GetVertices().Contains(SerializedGraph.Vertices[0]) == false);
+	CHECK(SplitIsland->GetVertices().Contains(SerializedGraph.Vertices[1]) == false);
+	CHECK(SplitIsland->GetVertices().Contains(SerializedGraph.Vertices[2]) == true);
+}
\ No newline at end of file
diff --git a/Engine/Plugins/Experimental/GameplayGraph/Tests/Private/GraphTests.cpp b/Engine/Plugins/Experimental/GameplayGraph/Tests/Private/GraphTests.cpp
new file mode 100644
--- /dev/null
+++ b/Engine/Plugins/Experimental/GameplayGraph/Tests/Private/GraphTests.cpp
@@ -0,0 +1,743 @@
+// Copyright Epic Games, Inc. All Rights Reserved.
+#include "TestHarness.h"
+#include "TestGraphBuilder.h"
+
+TEST_CASE("Graph::Properties::Equality", "[graph][properties]")
+{
+	SECTION("Equality")
+	{
+		FGraphProperties Props1;
+		Props1.bGenerateIslands = true;
+
+		FGraphProperties Props2;
+		Props2.bGenerateIslands = true;
+
+		CHECK(Props1 == Props2);
+	}
+
+	SECTION("InEquality")
+	{
+		FGraphProperties Props1;
+		Props1.bGenerateIslands = true;
+
+		FGraphProperties Props2;
+		Props2.bGenerateIslands = false;
+
+		CHECK(Props1 != Props2);
+	}
+}
+
+TEST_CASE("Graph::Default Constructor", "[graph]")
+{
+	TObjectPtr<UGraph> Graph = NewObject<UGraph>();
+
+	SECTION("State")
+	{
+		CHECK(Graph->NumVertices() == 0);
+		CHECK(Graph->NumEdges() == 0);
+		CHECK(Graph->NumIslands() == 0);
+	}
+}
+
+TEST_CASE("Graph::Initialize Properties", "[graph]")
+{
+	TObjectPtr<UGraph> Graph = NewObject<UGraph>();
+
+	SECTION("True Islands")
+	{
+		FGraphProperties Props1;
+		Props1.bGenerateIslands = true;
+		Graph->InitializeFromProperties(Props1);
+		CHECK(Graph->GetProperties() == Props1);
+	}
+
+	SECTION("False Islands")
+	{
+		FGraphProperties Props1;
+		Props1.bGenerateIslands = false;
+		Graph->InitializeFromProperties(Props1);
+		CHECK(Graph->GetProperties() == Props1);
+	}
+}
+
+TEST_CASE_METHOD(FTestGraphBuilder, "Graph::Create Vertices", "[graph]")
+{
+	FGraphUniqueIndex Node1Index = { FGuid { 1, 1, 1, 1 } };
+	FGraphVertexHandle Node1 = Graph->CreateVertex(Node1Index);
+	{
+		CHECK(Node1.GetUniqueIndex() == Node1Index);
+		CHECK(Graph->NumVertices() == 1);
+		CHECK(Graph->NumEdges() == 0);
+		CHECK(Graph->NumIslands() == 0);
+		CHECK(Graph->GetVertices().Contains(Node1) == true);
+	}
+
+	FGraphUniqueIndex Node2Index = { FGuid { 2, 2, 2, 2 } };
+	FGraphVertexHandle Node2 = Graph->CreateVertex(Node2Index);
+	{
+		CHECK(Node2.GetUniqueIndex() == Node2Index);
+		CHECK(Graph->NumVertices() == 2);
+		CHECK(Graph->NumEdges() == 0);
+		CHECK(Graph->NumIslands() == 0);
+		CHECK(Graph->GetVertices().Contains(Node2) == true);
+	}
+
+	{
+		{
+			FGraphVertexHandle Complete1 = Graph->GetCompleteNodeHandle(FGraphVertexHandle{ Node1Index, Graph });
+			CHECK(Complete1 == Node1);
+			CHECK(Complete1.IsValid() == true);
+			CHECK(Complete1.HasElement() == true);
+			CHECK(Complete1.IsComplete() == true);
+		}
+
+		{
+			FGraphVertexHandle Complete2 = Graph->GetCompleteNodeHandle(FGraphVertexHandle{ Node2Index, Graph });
+			CHECK(Complete2 == Node2);
+			CHECK(Complete2.IsValid() == true);
+			CHECK(Complete2.HasElement() == true);
+			CHECK(Complete2.IsComplete() == true);
+		}
+
+		{
+			FGraphVertexHandle Complete1 = Graph->GetCompleteNodeHandle(Node1);
+			CHECK(Complete1 == Node1);
+			CHECK(Complete1.IsValid() == true);
+			CHECK(Complete1.HasElement() == true);
+			CHECK(Complete1.IsComplete() == true);
+		}
+
+		{
+			FGraphVertexHandle Incomplete = Graph->GetCompleteNodeHandle(FGraphVertexHandle{ FGraphUniqueIndex{ FGuid{ 3, 3, 3, 3 } }, Graph });
+			CHECK(Incomplete.IsValid() == false);
+			CHECK(Incomplete.IsComplete() == false);
+			CHECK(Incomplete.HasElement() == false);
+		}
+
+		{
+			FGraphVertexHandle Incomplete = Graph->GetCompleteNodeHandle(FGraphVertexHandle{ FGraphUniqueIndex{}, Graph });
+			CHECK(Incomplete.IsValid() == false);
+			CHECK(Incomplete.IsComplete() == false);
+			CHECK(Incomplete.HasElement() == false);
+		}
+	}
+
+	{
+		Graph->FinalizeVertex(Node1);
+		CHECK(Graph->NumVertices() == 2);
+		CHECK(Graph->NumEdges() == 0);
+		CHECK(Graph->NumIslands() == 1);
+	}
+}
+
+TEST_CASE_METHOD(FTestGraphBuilder, "Graph::Create Vertices::Duplicate", "[graph]")
+{
+	FGraphUniqueIndex Node1Index = { FGuid { 1, 1, 1, 1 } };
+	FGraphVertexHandle Node1 = Graph->CreateVertex(Node1Index);
+	{
+		CHECK(Node1.GetUniqueIndex() == Node1Index);
+		CHECK(Graph->NumVertices() == 1);
+		CHECK(Graph->NumEdges() == 0);
+		CHECK(Graph->NumIslands() == 0);
+		CHECK(Graph->GetVertices().Contains(Node1) == true);
+	}
+
+	FGraphVertexHandle Node2 = Graph->CreateVertex(Node1Index);
+	{
+		CHECK(Node2.IsValid() == false);
+		CHECK(Node2.IsComplete() == false);
+		CHECK(Graph->NumVertices() == 1);
+		CHECK(Graph->NumEdges() == 0);
+		CHECK(Graph->NumIslands() == 0);
+		CHECK(Graph->GetVertices().Contains(Node2) == false);
+	}
+}
+
+TEST_CASE_METHOD(FTestGraphBuilder, "Graph::Create Edges::Single", "[graph]")
+{
+	PopulateVertices(4, false);
+
+	{
+		FEdgeCreationParameters Params;
+		Params.VertexHandle1 = VertexHandles[1];
+		Params.VertexHandle2 = VertexHandles[0];
+		Params.EdgeIndex = FGraphUniqueIndex{ FGuid{ 1, 1, 1, 1 } };
+
+		TArray<FGraphEdgeHandle> Edges;
+		Graph->CreateBulkEdges({ Params }, &Edges);
+		CHECK(Graph->NumEdges() == 1);
+		CHECK(Graph->NumIslands() == 1);
+
+		REQUIRE(Edges.Num() == 1);
+		CHECK(Edges[0].GetUniqueIndex() == Params.EdgeIndex);
+		CHECK(Edges[0].IsComplete());
+		CHECK(Graph->GetEdges().Contains(Edges[0]) == true);
+	}
+
+	{
+		FinalizeVertices();
+		CHECK(Graph->NumEdges() == 1);
+		CHECK(Graph->NumIslands() == 3);
+	}
+
+	{
+		FEdgeCreationParameters Params;
+		Params.VertexHandle1 = VertexHandles[2];
+		Params.VertexHandle2 = VertexHandles[3];
+		Params.EdgeIndex = FGraphUniqueIndex{ FGuid{ 2, 2, 2, 2 } };
+
+		TArray<FGraphEdgeHandle> Edges;
+		Graph->CreateBulkEdges({ Params }, &Edges);
+		CHECK(Graph->NumEdges() == 2);
+		CHECK(Graph->NumIslands() == 2);
+
+		REQUIRE(Edges.Num() == 1);
+		CHECK(Edges[0].GetUniqueIndex() == Params.EdgeIndex);
+		CHECK(Edges[0].IsComplete());
+		CHECK(Graph->GetEdges().Contains(Edges[0]) == true);
+	}
+
+	{
+		FEdgeCreationParameters Params;
+		Params.VertexHandle1 = VertexHandles[0];
+		Params.VertexHandle2 = VertexHandles[2];
+		Params.EdgeIndex = FGraphUniqueIndex{ FGuid{ 3, 3, 3, 3 } };
+
+		TArray<FGraphEdgeHandle> Edges;
+		Graph->CreateBulkEdges({ Params }, &Edges);
+		CHECK(Graph->NumEdges() == 3);
+		CHECK(Graph->NumIslands() == 1);
+
+		REQUIRE(Edges.Num() == 1);
+		CHECK(Edges[0].GetUniqueIndex() == Params.EdgeIndex);
+		CHECK(Graph->GetEdges().Contains(Edges[0]) == true);
+	}
+
+	{
+		FEdgeCreationParameters Params;
+		Params.VertexHandle1 = VertexHandles[0];
+		Params.VertexHandle2 = VertexHandles[2];
+		Params.EdgeIndex = FGraphUniqueIndex{ FGuid{ 4, 4, 4, 4 } };
+
+		TArray<FGraphEdgeHandle> Edges;
+		Graph->CreateBulkEdges({ Params }, &Edges);
+		CHECK(Graph->NumEdges() == 3);
+		CHECK(Graph->NumIslands() == 1);
+
+		REQUIRE(Edges.Num() == 0);
+	}
+
+	{
+		FEdgeCreationParameters Params;
+		Params.VertexHandle1 = VertexHandles[2];
+		Params.VertexHandle2 = VertexHandles[0];
+		Params.EdgeIndex = FGraphUniqueIndex{ FGuid{ 5, 5, 5, 5 } };
+
+		TArray<FGraphEdgeHandle> Edges;
+		Graph->CreateBulkEdges({ Params }, &Edges);
+		CHECK(Graph->NumEdges() == 3);
+		CHECK(Graph->NumIslands() == 1);
+
+		REQUIRE(Edges.Num() == 0);
+	}
+}
+
+TEST_CASE_METHOD(FTestGraphBuilder, "Graph::Create Edges::Bulk", "[graph]")
+{
+	PopulateVertices(4, true);
+
+	{
+		FEdgeCreationParameters Params1;
+		Params1.VertexHandle1 = VertexHandles[1];
+		Params1.VertexHandle2 = VertexHandles[0];
+		Params1.EdgeIndex = FGraphUniqueIndex{ FGuid{ 1, 1, 1, 1 } };
+
+		FEdgeCreationParameters Params2;
+		Params2.VertexHandle1 = VertexHandles[2];
+		Params2.VertexHandle2 = VertexHandles[3];
+		Params2.EdgeIndex = FGraphUniqueIndex{ FGuid{ 2, 2, 2, 2 } };
+
+		FEdgeCreationParameters Params3;
+		Params3.VertexHandle1 = VertexHandles[0];
+		Params3.VertexHandle2 = VertexHandles[2];
+		Params3.EdgeIndex = FGraphUniqueIndex{ FGuid{ 3, 3, 3, 3 } };
+
+		TArray<FGraphEdgeHandle> Edges;
+		Graph->CreateBulkEdges({ Params1, Params2, Params3 }, &Edges);
+		CHECK(Graph->NumEdges() == 3);
+		CHECK(Graph->NumIslands() == 1);
+
+		REQUIRE(Edges.Num() == 3);
+		CHECK(Edges[0].GetUniqueIndex() == Params1.EdgeIndex);
+		CHECK(Edges[1].GetUniqueIndex() == Params2.EdgeIndex);
+		CHECK(Edges[2].GetUniqueIndex() == Params3.EdgeIndex);
+		CHECK(Graph->GetEdges().Contains(Edges[0]) == true);
+		CHECK(Graph->GetEdges().Contains(Edges[1]) == true);
+		CHECK(Graph->GetEdges().Contains(Edges[2]) == true);
+		CHECK(Edges[0].IsComplete() == true);
+		CHECK(Edges[1].IsComplete() == true);
+		CHECK(Edges[2].IsComplete() == true);
+	}
+
+	{
+		FEdgeCreationParameters Params1;
+		Params1.VertexHandle1 = VertexHandles[0];
+		Params1.VertexHandle2 = VertexHandles[1];
+		Params1.EdgeIndex = FGraphUniqueIndex{ FGuid{ 4, 4, 4, 4 } };
+
+		FEdgeCreationParameters Params2;
+		Params2.VertexHandle1 = VertexHandles[3];
+		Params2.VertexHandle2 = VertexHandles[2];
+		Params2.EdgeIndex = FGraphUniqueIndex{ FGuid{ 5, 5, 5, 5 } };
+
+		FEdgeCreationParameters Params3;
+		Params3.VertexHandle1 = VertexHandles[2];
+		Params3.VertexHandle2 = VertexHandles[0];
+		Params3.EdgeIndex = FGraphUniqueIndex{ FGuid{ 6, 6, 6, 6 } };
+
+		TArray<FGraphEdgeHandle> Edges;
+		Graph->CreateBulkEdges({ Params1, Params2, Params3 }, &Edges);
+		CHECK(Graph->NumEdges() == 3);
+		CHECK(Graph->NumIslands() == 1);
+		CHECK(Edges.Num() == 0);
+	}
+}
+
+class FScopedVertexEdgeCleanupChecker
+{
+public:
+	template<typename TIterable>
+	FScopedVertexEdgeCleanupChecker(UGraph* InGraph, const TIterable& InVertexHandles)
+	: Graph(InGraph)
+	{
+		REQUIRE(Graph != nullptr);
+
+		for (const FGraphVertexHandle& VertexHandle : InVertexHandles)
+		{
+			UGraphVertex* Vertex = VertexHandle.GetVertex();
+			REQUIRE(Vertex != nullptr);
+
+			Vertex->ForEachAdjacentVertex(
+				[this](const FGraphVertexHandle& NeighborHandle, const FGraphEdgeHandle& EdgeHandle)
+				{
+					EdgeHandles.Add(EdgeHandle);
+					CHECK(Graph->GetEdges().Contains(EdgeHandle) == true);
+					CHECK(EdgeHandle.IsComplete() == true);
+				}
+			);
+		}
+	}
+
+	~FScopedVertexEdgeCleanupChecker()
+	{
+		for (const FGraphEdgeHandle& EdgeHandle : EdgeHandles)
+		{
+			CHECK(Graph->GetEdges().Contains(EdgeHandle) == false);
+			CHECK(EdgeHandle.IsComplete() == false);
+		}
+	}
+
+private:
+	UGraph* Graph;
+	TArray<FGraphEdgeHandle> EdgeHandles;
+};
+
+class FScopedVertexIslandCleanupChecker
+{
+public:
+	template<typename TIterable>
+	FScopedVertexIslandCleanupChecker(UGraph* InGraph, const TIterable& InVertexHandles)
+	: Graph(InGraph)
+	{
+		REQUIRE(Graph != nullptr);
+
+		for (const FGraphVertexHandle& VertexHandle : InVertexHandles)
+		{
+			UGraphVertex* Vertex = VertexHandle.GetVertex();
+			REQUIRE(Vertex != nullptr);
+
+			FGraphIslandHandle ParentIslandHandle = Vertex->GetParentIsland();
+			CHECK(Graph->GetIslands().Contains(ParentIslandHandle) == true);
+			CHECK(ParentIslandHandle.IsComplete() == true);
+			IslandHandles.Add(ParentIslandHandle);
+		}
+	}
+
+	~FScopedVertexIslandCleanupChecker()
+	{
+		for (const FGraphIslandHandle& IslandHandle : IslandHandles)
+		{
+			CHECK(Graph->GetIslands().Contains(IslandHandle) == false);
+			CHECK(IslandHandle.IsComplete() == false);
+		}
+	}
+
+private:
+	UGraph* Graph;
+	TArray<FGraphIslandHandle> IslandHandles;
+};
+
+TEST_CASE_METHOD(FTestGraphBuilder, "Graph::Remove Vertex::Single::Fully Connected", "[graph]")
+{
+	PopulateVertices(5, true);
+	BuildFullyConnectedEdges(5);
+
+	CHECK(Graph->NumVertices() == 5);
+	CHECK(Graph->NumEdges() == 10);
+	CHECK(Graph->NumIslands() == 1);
+
+	{
+		FScopedVertexEdgeCleanupChecker Checker { Graph, TArray{ VertexHandles[0] } };
+		Graph->RemoveVertex(VertexHandles[0]);
+	}
+
+	CHECK(VertexHandles[0].IsComplete() == false);
+	CHECK(Graph->GetCompleteNodeHandle(VertexHandles[0]).IsComplete() == false);
+	CHECK(Graph->GetVertices().Contains(VertexHandles[0]) == false);
+	CHECK(Graph->NumVertices() == 4);
+	CHECK(Graph->NumEdges() == 6);
+	CHECK(Graph->NumIslands() == 1);
+
+	{
+		FScopedVertexEdgeCleanupChecker Checker { Graph, TArray{ VertexHandles[1] } };
+		Graph->RemoveVertex(VertexHandles[1]);
+	}
+
+	CHECK(VertexHandles[1].IsComplete() == false);
+	CHECK(Graph->GetCompleteNodeHandle(VertexHandles[1]).IsComplete() == false);
+	CHECK(Graph->GetVertices().Contains(VertexHandles[1]) == false);
+	CHECK(Graph->NumVertices() == 3);
+	CHECK(Graph->NumEdges() == 3);
+	CHECK(Graph->NumIslands() == 1);
+
+	{
+		FScopedVertexEdgeCleanupChecker Checker { Graph, TArray{ VertexHandles[2] } };
+		Graph->RemoveVertex(VertexHandles[2]);
+	}
+
+	CHECK(VertexHandles[2].IsComplete() == false);
+	CHECK(Graph->GetCompleteNodeHandle(VertexHandles[2]).IsComplete() == false);
+	CHECK(Graph->GetVertices().Contains(VertexHandles[2]) == false);
+	CHECK(Graph->NumVertices() == 2);
+	CHECK(Graph->NumEdges() == 1);
+	CHECK(Graph->NumIslands() == 1);
+
+	{
+		FScopedVertexEdgeCleanupChecker Checker { Graph, TArray{ VertexHandles[3] } };
+		Graph->RemoveVertex(VertexHandles[3]);
+	}
+
+	CHECK(VertexHandles[3].IsComplete() == false);
+	CHECK(Graph->GetCompleteNodeHandle(VertexHandles[3]).IsComplete() == false);
+	CHECK(Graph->GetVertices().Contains(VertexHandles[3]) == false);
+	CHECK(Graph->NumVertices() == 1);
+	CHECK(Graph->NumEdges() == 0);
+	CHECK(Graph->NumIslands() == 1);
+
+	{
+		FScopedVertexEdgeCleanupChecker Checker { Graph, TArray{ VertexHandles[4] } };
+		Graph->RemoveVertex(VertexHandles[4]);
+	}
+
+	CHECK(VertexHandles[4].IsComplete() == false);
+	CHECK(Graph->GetCompleteNodeHandle(VertexHandles[4]).IsComplete() == false);
+	CHECK(Graph->GetVertices().Contains(VertexHandles[4]) == false);
+	CHECK(Graph->NumVertices() == 0);
+	CHECK(Graph->NumEdges() == 0);
+	CHECK(Graph->NumIslands() == 0);
+}
+
+TEST_CASE_METHOD(FTestGraphBuilder, "Graph::Remove Vertex::Single::Linear", "[graph]")
+{
+	PopulateVertices(10, true);
+	BuildLinearEdges(10);
+
+	CHECK(Graph->NumVertices() == 10);
+	CHECK(Graph->NumEdges() == 9);
+	CHECK(Graph->NumIslands() == 1);
+
+	SECTION("First Vertex")
+	{
+		CHECK(Graph->GetVertices().Contains(VertexHandles[0]) == true);
+		CHECK(VertexHandles[0].IsComplete() == true);
+
+		FScopedVertexEdgeCleanupChecker Checker { Graph, TArray{ VertexHandles[0] } };
+		Graph->RemoveVertex(VertexHandles[0]);
+
+		CHECK(Graph->GetVertices().Contains(VertexHandles[0]) == false);
+		CHECK(VertexHandles[0].IsComplete() == false);
+
+		CHECK(Graph->NumVertices() == 9);
+		CHECK(Graph->NumEdges() == 8);
+		CHECK(Graph->NumIslands() == 1);
+	}
+
+	SECTION("Split Island")
+	{
+		CHECK(Graph->GetVertices().Contains(VertexHandles[5]) == true);
+		CHECK(VertexHandles[5].IsComplete() == true);
+
+		FScopedVertexEdgeCleanupChecker Checker { Graph, TArray{ VertexHandles[5] } };
+		Graph->RemoveVertex(VertexHandles[5]);
+
+		CHECK(Graph->GetVertices().Contains(VertexHandles[5]) == false);
+		CHECK(VertexHandles[5].IsComplete() == false);
+
+		CHECK(Graph->NumVertices() == 9);
+		CHECK(Graph->NumEdges() == 7);
+		CHECK(Graph->NumIslands() == 2);
+	}
+}
+
+TEST_CASE_METHOD(FTestGraphBuilder, "Graph::Remove Vertex::Bulk::Fully Connected", "[graph]")
+{
+	PopulateVertices(5, true);
+	BuildFullyConnectedEdges(5);
+
+	CHECK(Graph->NumVertices() == 5);
+	CHECK(Graph->NumEdges() == 10);
+	CHECK(Graph->NumIslands() == 1);
+
+	{
+		FScopedVertexEdgeCleanupChecker Checker { Graph, TArray{ VertexHandles[0], VertexHandles[1] } };
+		Graph->RemoveBulkVertices( { VertexHandles[0], VertexHandles[1] });
+	}
+
+	CHECK(Graph->GetCompleteNodeHandle(VertexHandles[0]).IsComplete() == false);
+	CHECK(Graph->GetCompleteNodeHandle(VertexHandles[1]).IsComplete() == false);
+	CHECK(VertexHandles[0].IsComplete() == false);
+	CHECK(VertexHandles[1].IsComplete() == false);
+	CHECK(Graph->NumVertices() == 3);
+	CHECK(Graph->NumEdges() == 3);
+	CHECK(Graph->NumIslands() == 1);
+
+	{
+		FScopedVertexEdgeCleanupChecker Checker { Graph, TArray{ VertexHandles[2], VertexHandles[3], VertexHandles[4] } };
+		Graph->RemoveBulkVertices( { VertexHandles[2], VertexHandles[3], VertexHandles[4] });
+	}
+
+	CHECK(Graph->GetCompleteNodeHandle(VertexHandles[2]).IsComplete() == false);
+	CHECK(Graph->GetCompleteNodeHandle(VertexHandles[3]).IsComplete() == false);
+	CHECK(Graph->GetCompleteNodeHandle(VertexHandles[4]).IsComplete() == false);
+	CHECK(VertexHandles[2].IsComplete() == false);
+	CHECK(VertexHandles[3].IsComplete() == false);
+	CHECK(VertexHandles[4].IsComplete() == false);
+	CHECK(Graph->NumVertices() == 0);
+	CHECK(Graph->NumEdges() == 0);
+	CHECK(Graph->NumIslands() == 0);
+}
+
+
+TEST_CASE_METHOD(FTestGraphBuilder, "Graph::Remove Vertex::Bulk::Linear", "[graph]")
+{
+	PopulateVertices(10, true);
+	BuildLinearEdges(10);
+
+	CHECK(Graph->NumVertices() == 10);
+	CHECK(Graph->NumEdges() == 9);
+	CHECK(Graph->NumIslands() == 1);
+
+	TArray<FGraphVertexHandle> VerticesToRemove = { VertexHandles[1], VertexHandles[2], VertexHandles[8] };
+	for (const FGraphVertexHandle& VertexHandle : VerticesToRemove)
+	{
+		CHECK(Graph->GetVertices().Contains(VertexHandle) == true);
+		CHECK(VertexHandle.IsComplete() == true);
+	}
+	FScopedVertexEdgeCleanupChecker Checker { Graph, VerticesToRemove };
+
+	Graph->RemoveBulkVertices(VerticesToRemove);
+	for (const FGraphVertexHandle& VertexHandle : VerticesToRemove)
+	{
+		CHECK(Graph->GetVertices().Contains(VertexHandle) == false);
+		CHECK(VertexHandle.IsComplete() == false);
+	}
+
+	CHECK(Graph->NumVertices() == 7);
+	CHECK(Graph->NumEdges() == 4);
+	CHECK(Graph->NumIslands() == 3);
+}
+
+TEST_CASE_METHOD(FTestGraphBuilder, "Graph::Remove Island", "[graph]")
+{
+	PopulateVertices(10, true);
+	BuildFullyConnectedEdges(5);
+
+	CHECK(Graph->NumVertices() == 10);
+	CHECK(Graph->NumEdges() == 20);
+	CHECK(Graph->NumIslands() == 2);
+
+	{
+		CHECK(Graph->GetIslands().Contains(IslandHandles[0]) == true);
+		CHECK(IslandHandles[0].IsComplete() == true);
+		TSet<FGraphVertexHandle> IslandVertices = IslandHandles[0].GetIsland()->GetVertices();
+		FScopedVertexEdgeCleanupChecker EdgeChecker { Graph, IslandVertices };
+		FScopedVertexIslandCleanupChecker IslandChecker { Graph, IslandVertices };
+		for (const FGraphVertexHandle& VertexHandle : IslandVertices)
+		{
+			CHECK(Graph->GetCompleteNodeHandle(VertexHandle).IsComplete() == true);
+			CHECK(Graph->GetVertices().Contains(VertexHandle) == true);
+		}
+		Graph->RemoveIsland(IslandHandles[0]);
+		for (const FGraphVertexHandle& VertexHandle : IslandVertices)
+		{
+			CHECK(Graph->GetCompleteNodeHandle(VertexHandle).IsComplete() == false);
+			CHECK(Graph->GetVertices().Contains(VertexHandle) == false);
+		}
+
+		CHECK(Graph->GetIslands().Contains(IslandHandles[0]) == false);
+		CHECK(IslandHandles[0].IsComplete() == false);
+	}
+
+	CHECK(Graph->NumVertices() == 5);
+	CHECK(Graph->NumEdges() == 10);
+	CHECK(Graph->NumIslands() == 1);
+
+	{
+		CHECK(Graph->GetIslands().Contains(IslandHandles[1]) == true);
+		CHECK(IslandHandles[1].IsComplete() == true);
+		TSet<FGraphVertexHandle> IslandVertices = IslandHandles[1].GetIsland()->GetVertices();
+		FScopedVertexEdgeCleanupChecker Checker { Graph, IslandVertices };
+		FScopedVertexIslandCleanupChecker IslandChecker { Graph, IslandVertices };
+		for (const FGraphVertexHandle& VertexHandle : IslandVertices)
+		{
+			CHECK(Graph->GetCompleteNodeHandle(VertexHandle).IsComplete() == true);
+			CHECK(Graph->GetVertices().Contains(VertexHandle) == true);
+		}
+		Graph->RemoveIsland(IslandHandles[1]);
+		for (const FGraphVertexHandle& VertexHandle : IslandVertices)
+		{
+			CHECK(Graph->GetCompleteNodeHandle(VertexHandle).IsComplete() == false);
+			CHECK(Graph->GetVertices().Contains(VertexHandle) == false);
+		}
+
+		CHECK(Graph->GetIslands().Contains(IslandHandles[1]) == false);
+		CHECK(IslandHandles[1].IsComplete() == false);
+	}
+
+	CHECK(Graph->NumVertices() == 0);
+	CHECK(Graph->NumEdges() == 0);
+	CHECK(Graph->NumIslands() == 0);
+}
+
+TEST_CASE_METHOD(FTestGraphBuilder, "Graph::Remove Edge::Handle Islands", "[graph]")
+{
+	PopulateVertices(10, true);
+	BuildLinearEdges(10);
+
+	CHECK(Graph->NumVertices() == 10);
+	CHECK(Graph->NumEdges() == 9);
+	CHECK(Graph->NumIslands() == 1);
+
+	SECTION("First Vertex")
+	{
+		TArray<FGraphEdgeHandle> Edges = GetEdgesForVertex(VertexHandles[0]);
+		for (const FGraphEdgeHandle& EdgeHandle : Edges)
+		{
+			CHECK(Graph->GetEdges().Contains(EdgeHandle) == true);
+			CHECK(EdgeHandle.IsComplete() == true);
+
+			Graph->RemoveEdge(EdgeHandle);
+
+			CHECK(Graph->GetEdges().Contains(EdgeHandle) == false);
+			CHECK(EdgeHandle.IsComplete() == false);
+		}
+
+		CHECK(Graph->NumVertices() == 10);
+		CHECK(Graph->NumEdges() == 8);
+		CHECK(Graph->NumIslands() == 2);
+	}
+
+	SECTION("Split Island")
+	{
+		TArray<FGraphEdgeHandle> Edges = GetEdgesForVertex(VertexHandles[5]);
+		for (const FGraphEdgeHandle& EdgeHandle : Edges)
+		{
+			CHECK(Graph->GetEdges().Contains(EdgeHandle) == true);
+			CHECK(EdgeHandle.IsComplete() == true);
+
+			Graph->RemoveEdge(EdgeHandle);
+
+			CHECK(Graph->GetEdges().Contains(EdgeHandle) == false);
+			CHECK(EdgeHandle.IsComplete() == false);
+		}
+
+		CHECK(Graph->NumVertices() == 10);
+		CHECK(Graph->NumEdges() == 7);
+		CHECK(Graph->NumIslands() == 3);
+	}
+}
+
+TEST_CASE_METHOD(FTestGraphBuilder, "Graph::Event::Vertex Created", "[graph]")
+{
+	FGraphUniqueIndex Node1Index = { FGuid { 1, 1, 1, 1 } };
+
+	FGraphVertexHandle CreatedHandle;
+	Graph->OnVertexCreated.AddLambda(
+		[&CreatedHandle, &Node1Index](const FGraphVertexHandle& InCreatedHandle)
+		{
+			CHECK(InCreatedHandle.IsComplete());
+			CHECK(InCreatedHandle.GetUniqueIndex() == Node1Index);
+			CreatedHandle = InCreatedHandle;
+		}
+	);
+
+	FGraphVertexHandle Node1 = Graph->CreateVertex(Node1Index);
+	CHECK(CreatedHandle == Node1);
+	CHECK(CreatedHandle.IsComplete() == true);
+}
+
+TEST_CASE_METHOD(FTestGraphBuilder, "Graph::Event::Edge Created::Single", "[graph]")
+{
+	PopulateVertices(10, true);
+
+	FEdgeCreationParameters Params;
+	Params.VertexHandle1 = VertexHandles[1];
+	Params.VertexHandle2 = VertexHandles[0];
+	Params.EdgeIndex = FGraphUniqueIndex{ FGuid{ 1, 1, 1, 1 } };
+
+	FGraphEdgeHandle CreatedHandle;
+	Graph->OnEdgeCreated.AddLambda(
+		[&CreatedHandle, &Params](const FGraphEdgeHandle& InCreatedHandle)
+		{
+			CHECK(InCreatedHandle.IsComplete());
+			CHECK(InCreatedHandle.GetUniqueIndex() == Params.EdgeIndex);
+			CreatedHandle = InCreatedHandle;
+		}
+	);
+
+	TArray<FGraphEdgeHandle> Edges;
+	Graph->CreateBulkEdges({ Params }, &Edges);
+
+	REQUIRE(Edges.Num() == 1);
+	CHECK(Edges[0] == CreatedHandle);
+}
+
+TEST_CASE_METHOD(FTestGraphBuilder, "Graph::Event::Island Created", "[graph]")
+{
+	PopulateVertices(10, false);
+
+	FEdgeCreationParameters Params;
+	Params.VertexHandle1 = VertexHandles[1];
+	Params.VertexHandle2 = VertexHandles[0];
+	Params.EdgeIndex = FGraphUniqueIndex{ FGuid{ 1, 1, 1, 1 } };
+
+	FGraphIslandHandle CreatedHandle;
+	Graph->OnIslandCreated.AddLambda(
+		[&CreatedHandle](const FGraphIslandHandle& InCreatedHandle)
+		{
+			CHECK(InCreatedHandle.IsComplete());
+			CreatedHandle = InCreatedHandle;
+		}
+	);
+
+	Graph->CreateBulkEdges({ Params });
+	CHECK(CreatedHandle.IsComplete());
+
+	UGraphVertex* Vertex1 = VertexHandles[1].GetVertex();
+	REQUIRE(Vertex1 != nullptr);
+	CHECK(Vertex1->GetParentIsland() == CreatedHandle);
+
+	UGraphVertex* Vertex0 = VertexHandles[0].GetVertex();
+	REQUIRE(Vertex0 != nullptr);
+	CHECK(Vertex0->GetParentIsland() == CreatedHandle);
+}
\ No newline at end of file
diff --git a/Engine/Plugins/Experimental/GameplayGraph/Tests/Private/GraphVertexTests.cpp b/Engine/Plugins/Experimental/GameplayGraph/Tests/Private/GraphVertexTests.cpp
new file mode 100644
--- /dev/null
+++ b/Engine/Plugins/Experimental/GameplayGraph/Tests/Private/GraphVertexTests.cpp
@@ -0,0 +1,512 @@
+// Copyright Epic Games, Inc. All Rights Reserved.
+#include "TestHarness.h"
+#include "TestGraphBuilder.h"
+
+TEST_CASE_METHOD(FTestGraphBuilder, "Graph::Vertex::Create", "[graph][vertex]")
+{
+	FGraphUniqueIndex VertexIndex = { FGuid { 1, 1, 1, 1 } };
+	FGraphVertexHandle VertexHandle = Graph->CreateVertex(VertexIndex);
+
+	CHECK(VertexHandle.IsValid());
+	CHECK(VertexHandle.IsComplete());
+	CHECK(VertexHandle.HasElement() == true);
+	CHECK(VertexHandle.GetUniqueIndex() == VertexIndex);
+	CHECK(VertexHandle.GetGraph() == Graph);
+
+	UGraphVertex* Vertex = VertexHandle.GetVertex();
+	REQUIRE(Vertex != nullptr);
+
+	CHECK(Vertex->Handle() == VertexHandle);
+	CHECK(Vertex->NumEdges() == 0);
+	CHECK(Vertex->GetParentIsland() == FGraphIslandHandle{});
+
+	Graph->FinalizeVertex(VertexHandle);
+
+	CHECK(Vertex->GetParentIsland() == Graph->GetIslands().CreateConstIterator()->Key);
+}
+
+TEST_CASE_METHOD(FTestGraphBuilder, "Graph::Vertex::Has Edge To", "[graph][vertex]")
+{
+	constexpr int32 NumVertices = 7;
+	PopulateVertices(NumVertices, true);
+	BuildFullyConnectedEdges(NumVertices);
+
+	SECTION("Post Construction")
+	{
+		for (int32 Index1 = 0; Index1 < NumVertices; ++Index1)
+		{
+			UGraphVertex* Vertex1 = VertexHandles[Index1].GetVertex();
+			REQUIRE(Vertex1 != nullptr);
+
+			for (int32 Index2 = Index1 + 1; Index2 < NumVertices; ++Index2)
+			{
+				UGraphVertex* Vertex2 = VertexHandles[Index2].GetVertex();
+				REQUIRE(Vertex2 != nullptr);
+
+				CHECK(Vertex1->HasEdgeTo(VertexHandles[Index2]) == true);
+				CHECK(Vertex1->GetEdges().Contains(VertexHandles[Index2]) == true);
+
+				CHECK(Vertex2->HasEdgeTo(VertexHandles[Index1]) == true);
+				CHECK(Vertex2->GetEdges().Contains(VertexHandles[Index1]) == true);
+			}
+		}
+	}
+
+	SECTION("Post remove vertex")
+	{
+		Graph->RemoveVertex(VertexHandles[0]);
+		for (int32 Index1 = 1; Index1 < NumVertices; ++Index1)
+		{
+			UGraphVertex* Vertex1 = VertexHandles[Index1].GetVertex();
+			REQUIRE(Vertex1 != nullptr);
+
+			for (int32 Index2 = Index1 + 1; Index2 < NumVertices; ++Index2)
+			{
+				UGraphVertex* Vertex2 = VertexHandles[Index2].GetVertex();
+				REQUIRE(Vertex2 != nullptr);
+
+				CHECK(Vertex1->HasEdgeTo(VertexHandles[Index2]) == true);
+				CHECK(Vertex2->HasEdgeTo(VertexHandles[Index1]) == true);
+			}
+
+			CHECK(Vertex1->HasEdgeTo(VertexHandles[0]) == false);
+			CHECK(Vertex1->GetEdges().Contains(VertexHandles[0]) == false);
+		}
+	}
+
+	SECTION("Post remove edge")
+	{
+		{
+			UGraphVertex* Vertex0 = VertexHandles[0].GetVertex();
+			REQUIRE(Vertex0 != nullptr);
+			UGraphVertex* Vertex1 = VertexHandles[1].GetVertex();
+			REQUIRE(Vertex1 != nullptr);
+
+			const FGraphEdgeHandle* Edge = Vertex0->FindEdgeTo(VertexHandles[1]);
+			REQUIRE(Edge != nullptr);
+
+			CHECK(Edge->IsComplete() == true);
+			Graph->RemoveEdge(*Edge);
+			CHECK(Edge->IsComplete() == false);
+
+			CHECK(Vertex0->HasEdgeTo(VertexHandles[1]) == false);
+			CHECK(Vertex1->HasEdgeTo(VertexHandles[0]) == false);
+		}
+
+		{
+			for (int32 Index1 = 0; Index1 < NumVertices; ++Index1)
+			{
+				UGraphVertex* Vertex1 = VertexHandles[Index1].GetVertex();
+				REQUIRE(Vertex1 != nullptr);
+
+				for (int32 Index2 = Index1 + 1; Index2 < NumVertices; ++Index2)
+				{
+					if (Index1 == 0 && Index2 == 1)
+					{
+						continue;
+					}
+
+					UGraphVertex* Vertex2 = VertexHandles[Index2].GetVertex();
+					REQUIRE(Vertex2 != nullptr);
+
+					CHECK(Vertex1->HasEdgeTo(VertexHandles[Index2]) == true);
+					CHECK(Vertex1->GetEdges().Contains(VertexHandles[Index2]) == true);
+
+					CHECK(Vertex2->HasEdgeTo(VertexHandles[Index1]) == true);
+					CHECK(Vertex2->GetEdges().Contains(VertexHandles[Index1]) == true);
+				}
+			}
+		}
+	}
+}
+
+TEST_CASE_METHOD(FTestGraphBuilder, "Graph::Vertex::Find Edge To", "[graph][vertex]")
+{
+	constexpr int32 NumVertices = 7;
+	PopulateVertices(NumVertices, true);
+	BuildFullyConnectedEdges(NumVertices);
+
+	SECTION("Post Construction")
+	{
+		for (int32 Index1 = 0; Index1 < NumVertices; ++Index1)
+		{
+			UGraphVertex* Vertex1 = VertexHandles[Index1].GetVertex();
+			REQUIRE(Vertex1 != nullptr);
+
+			for (int32 Index2 = Index1 + 1; Index2 < NumVertices; ++Index2)
+			{
+				UGraphVertex* Vertex2 = VertexHandles[Index2].GetVertex();
+				REQUIRE(Vertex2 != nullptr);
+			
+				const FGraphEdgeHandle* Edge1 = Vertex1->FindEdgeTo(VertexHandles[Index2]);
+				REQUIRE(Edge1 != nullptr);
+				const FGraphEdgeHandle* Edge2 = Vertex2->FindEdgeTo(VertexHandles[Index1]);
+				REQUIRE(Edge2 != nullptr);
+				CHECK(*Edge1 == *Edge2);
+			}
+		}
+	}
+
+	SECTION("Post remove vertex")
+	{
+		Graph->RemoveVertex(VertexHandles[0]);
+		for (int32 Index1 = 1; Index1 < NumVertices; ++Index1)
+		{
+			UGraphVertex* Vertex1 = VertexHandles[Index1].GetVertex();
+			REQUIRE(Vertex1 != nullptr);
+
+			for (int32 Index2 = Index1 + 1; Index2 < NumVertices; ++Index2)
+			{
+				UGraphVertex* Vertex2 = VertexHandles[Index2].GetVertex();
+				REQUIRE(Vertex2 != nullptr);
+
+				const FGraphEdgeHandle* Edge1 = Vertex1->FindEdgeTo(VertexHandles[Index2]);
+				REQUIRE(Edge1 != nullptr);
+				const FGraphEdgeHandle* Edge2 = Vertex2->FindEdgeTo(VertexHandles[Index1]);
+				REQUIRE(Edge2 != nullptr);
+				CHECK(*Edge1 == *Edge2);
+			}
+
+			CHECK(Vertex1->FindEdgeTo(VertexHandles[0]) == nullptr);
+		}
+	}
+
+	SECTION("Post remove edge")
+	{
+		{
+			UGraphVertex* Vertex0 = VertexHandles[0].GetVertex();
+			REQUIRE(Vertex0 != nullptr);
+			UGraphVertex* Vertex1 = VertexHandles[1].GetVertex();
+			REQUIRE(Vertex1 != nullptr);
+
+			const FGraphEdgeHandle* Edge = Vertex0->FindEdgeTo(VertexHandles[1]);
+			REQUIRE(Edge != nullptr);
+
+			CHECK(Edge->IsComplete() == true);
+			Graph->RemoveEdge(*Edge);
+			CHECK(Edge->IsComplete() == false);
+
+			CHECK(Vertex0->FindEdgeTo(VertexHandles[1]) == nullptr);
+			CHECK(Vertex1->FindEdgeTo(VertexHandles[0]) == nullptr);
+		}
+
+		{
+			for (int32 Index1 = 0; Index1 < NumVertices; ++Index1)
+			{
+				UGraphVertex* Vertex1 = VertexHandles[Index1].GetVertex();
+				REQUIRE(Vertex1 != nullptr);
+
+				for (int32 Index2 = Index1 + 1; Index2 < NumVertices; ++Index2)
+				{
+					if (Index1 == 0 && Index2 == 1)
+					{
+						continue;
+					}
+
+					UGraphVertex* Vertex2 = VertexHandles[Index2].GetVertex();
+					REQUIRE(Vertex2 != nullptr);
+
+					const FGraphEdgeHandle* Edge1 = Vertex1->FindEdgeTo(VertexHandles[Index2]);
+					REQUIRE(Edge1 != nullptr);
+					const FGraphEdgeHandle* Edge2 = Vertex2->FindEdgeTo(VertexHandles[Index1]);
+					REQUIRE(Edge2 != nullptr);
+					CHECK(*Edge1 == *Edge2);
+				}
+			}
+		}
+	}
+}
+
+TEST_CASE_METHOD(FTestGraphBuilder, "Graph::Vertex::Num Edges", "[graph][vertex]")
+{
+	constexpr int32 NumVertices = 7;
+	PopulateVertices(NumVertices, true);
+	BuildFullyConnectedEdges(NumVertices);
+
+	SECTION("Post Construction")
+	{
+		for (int32 Index1 = 0; Index1 < NumVertices; ++Index1)
+		{
+			UGraphVertex* Vertex1 = VertexHandles[Index1].GetVertex();
+			REQUIRE(Vertex1 != nullptr);
+			CHECK(Vertex1->NumEdges() == 6);
+		}
+	}
+
+	SECTION("Post remove vertex")
+	{
+		Graph->RemoveVertex(VertexHandles[0]);
+		for (int32 Index1 = 1; Index1 < NumVertices; ++Index1)
+		{
+			UGraphVertex* Vertex1 = VertexHandles[Index1].GetVertex();
+			REQUIRE(Vertex1 != nullptr);
+			CHECK(Vertex1->NumEdges() == 5);
+		}
+	}
+
+	SECTION("Post remove edge")
+	{
+		UGraphVertex* Vertex0 = VertexHandles[0].GetVertex();
+		REQUIRE(Vertex0 != nullptr);
+
+		const FGraphEdgeHandle* Edge = Vertex0->FindEdgeTo(VertexHandles[1]);
+		REQUIRE(Edge != nullptr);
+
+		CHECK(Edge->IsComplete() == true);
+		Graph->RemoveEdge(*Edge);
+		CHECK(Edge->IsComplete() == false);
+
+		for (int32 Index1 = 0; Index1 < NumVertices; ++Index1)
+		{
+			UGraphVertex* Vertex1 = VertexHandles[Index1].GetVertex();
+			REQUIRE(Vertex1 != nullptr);
+
+			if (Index1 == 0 || Index1 == 1)
+			{
+				CHECK(Vertex1->NumEdges() == 5);
+			}
+			else
+			{
+				CHECK(Vertex1->NumEdges() == 6);
+			}
+		}
+	}
+}
+
+TEST_CASE_METHOD(FTestGraphBuilder, "Graph::Vertex::Get Parent Island", "[graph][vertex]")
+{
+	constexpr int32 NumVertices = 7;
+	PopulateVertices(NumVertices, true);
+	BuildLinearEdges(NumVertices);
+
+	SECTION("Post Construction")
+	{
+		for (int32 Index1 = 1; Index1 < NumVertices; ++Index1)
+		{
+			UGraphVertex* Vertex0 = VertexHandles[Index1-1].GetVertex();
+			REQUIRE(Vertex0 != nullptr);
+			UGraphVertex* Vertex1 = VertexHandles[Index1].GetVertex();
+			REQUIRE(Vertex1 != nullptr);
+			CHECK(Vertex0->GetParentIsland() == Vertex1->GetParentIsland());
+			CHECK(Vertex0->GetParentIsland() == Graph->GetIslands().CreateConstIterator()->Key);
+		}
+	}
+
+	SECTION("Post remove vertex")
+	{
+		Graph->RemoveVertex(VertexHandles[3]);
+		for (int32 Index1 = 1; Index1 < 3; ++Index1)
+		{
+			UGraphVertex* Vertex0 = VertexHandles[Index1-1].GetVertex();
+			REQUIRE(Vertex0 != nullptr);
+			UGraphVertex* Vertex1 = VertexHandles[Index1].GetVertex();
+			REQUIRE(Vertex1 != nullptr);
+			CHECK(Vertex0->GetParentIsland() == Vertex1->GetParentIsland());
+		}
+
+		{
+			UGraphVertex* Vertex0 = VertexHandles[2].GetVertex();
+			REQUIRE(Vertex0 != nullptr);
+			UGraphVertex* Vertex1 = VertexHandles[4].GetVertex();
+			REQUIRE(Vertex1 != nullptr);
+			CHECK(Vertex0->GetParentIsland() != Vertex1->GetParentIsland());
+		}
+
+		for (int32 Index1 = 5; Index1 < NumVertices; ++Index1)
+		{
+			UGraphVertex* Vertex0 = VertexHandles[Index1-1].GetVertex();
+			REQUIRE(Vertex0 != nullptr);
+			UGraphVertex* Vertex1 = VertexHandles[Index1].GetVertex();
+			REQUIRE(Vertex1 != nullptr);
+			CHECK(Vertex0->GetParentIsland() == Vertex1->GetParentIsland());
+		}
+	}
+
+	SECTION("Post remove edge")
+	{
+		UGraphVertex* Vertex0 = VertexHandles[3].GetVertex();
+		REQUIRE(Vertex0 != nullptr);
+
+		const FGraphEdgeHandle* Edge = Vertex0->FindEdgeTo(VertexHandles[4]);
+		REQUIRE(Edge != nullptr);
+
+		CHECK(Edge->IsComplete() == true);
+		Graph->RemoveEdge(*Edge);
+		CHECK(Edge->IsComplete() == false);
+
+		CHECK(Vertex0->GetParentIsland() != VertexHandles[4].GetVertex()->GetParentIsland());
+		CHECK(Vertex0->GetParentIsland() == VertexHandles[2].GetVertex()->GetParentIsland());
+	}
+}
+
+TEST_CASE_METHOD(FTestGraphBuilder, "Graph::Vertex::For Each Adjacent::With Edge", "[graph][vertex]")
+{
+	constexpr int32 NumVertices = 7;
+	PopulateVertices(NumVertices, true);
+	BuildFullyConnectedEdges(NumVertices);
+
+	for (int32 Index1 = 0; Index1 < NumVertices; ++Index1)
+	{
+		UGraphVertex* Vertex1 = VertexHandles[Index1].GetVertex();
+		REQUIRE(Vertex1 != nullptr);
+
+		TMap<FGraphVertexHandle, FGraphEdgeHandle> ConnectedVertices;
+		ConnectedVertices.Reserve(NumVertices - 1);
+
+		Vertex1->ForEachAdjacentVertex(
+			[&ConnectedVertices](const FGraphVertexHandle& NeighborHandle, const FGraphEdgeHandle& EdgeHandle)
+			{
+				ConnectedVertices.Add(NeighborHandle, EdgeHandle);
+			}
+		);
+
+		CHECK(ConnectedVertices == Vertex1->GetEdges());
+		CHECK(ConnectedVertices.Num() == NumVertices - 1);
+		for (int32 Index2 = 0; Index2 < NumVertices; ++Index2)
+		{
+			if (Index1 == Index2)
+			{
+				continue;
+			}
+			
+			CHECK(ConnectedVertices.Contains(VertexHandles[Index2]));
+			FGraphEdgeHandle* EdgeHandle = ConnectedVertices.Find(VertexHandles[Index2]);
+			REQUIRE(EdgeHandle != nullptr);
+			CHECK(Vertex1->GetEdges().FindRef(VertexHandles[Index2]) == *EdgeHandle);
+			ConnectedVertices.Remove(VertexHandles[Index2]);
+		}
+		CHECK(ConnectedVertices.Num() == 0);
+	}
+}
+
+TEST_CASE_METHOD(FTestGraphBuilder, "Graph::Vertex::For Each Adjacent::Without Edge", "[graph][vertex]")
+{
+	constexpr int32 NumVertices = 7;
+	PopulateVertices(NumVertices, true);
+	BuildFullyConnectedEdges(NumVertices);
+
+	for (int32 Index1 = 0; Index1 < NumVertices; ++Index1)
+	{
+		UGraphVertex* Vertex1 = VertexHandles[Index1].GetVertex();
+		REQUIRE(Vertex1 != nullptr);
+
+		TSet<FGraphVertexHandle> ConnectedVertices;
+		ConnectedVertices.Reserve(NumVertices - 1);
+
+		Vertex1->ForEachAdjacentVertex(
+			[&ConnectedVertices](const FGraphVertexHandle& NeighborHandle)
+			{
+				ConnectedVertices.Add(NeighborHandle);
+			}
+		);
+
+		CHECK(ConnectedVertices.Num() == NumVertices - 1);
+		for (int32 Index2 = 0; Index2 < NumVertices; ++Index2)
+		{
+			if (Index1 == Index2)
+			{
+				continue;
+			}
+			
+			CHECK(ConnectedVertices.Contains(VertexHandles[Index2]));
+			ConnectedVertices.Remove(VertexHandles[Index2]);
+		}
+		CHECK(ConnectedVertices.Num() == 0);
+	}
+}
+
+TEST_CASE_METHOD(FTestGraphBuilder, "Graph::Vertex::Event::Remove", "[graph][vertex]")
+{
+	constexpr int32 NumVertices = 7;
+	PopulateVertices(NumVertices, false);
+
+	TArray<bool> VertexRemoved;
+	VertexRemoved.Reserve(NumVertices);
+
+	int32 Index = 0;
+	for (const FGraphVertexHandle& VertexHandle : VertexHandles)
+	{
+		VertexRemoved.Add(false);
+
+		UGraphVertex* Vertex = VertexHandle.GetVertex();
+		REQUIRE(Vertex != nullptr);
+		Vertex->OnVertexRemoved.AddLambda(
+			[Index, &VertexRemoved]()
+			{
+				VertexRemoved[Index] = true;
+			}
+		);
+
+		++Index;
+	}
+
+	for (Index = 0; Index < NumVertices; ++Index)
+	{
+		Graph->RemoveVertex(VertexHandles[Index]);
+
+		for (int32 Index2 = 0; Index2 < NumVertices; ++Index2)
+		{
+			if (Index2 <= Index)
+			{
+				CHECK(VertexRemoved[Index2] == true);
+			}
+			else
+			{
+				CHECK(VertexRemoved[Index2] == false);
+			}
+		}
+	}
+}
+
+TEST_CASE_METHOD(FTestGraphBuilder, "Graph::Vertex::Event::Parent Island Set", "[graph][vertex]")
+{
+	constexpr int32 NumVertices = 7;
+	PopulateVertices(NumVertices, false);
+
+	TArray<FGraphIslandHandle> VertexIslandSet;
+	VertexIslandSet.Reserve(NumVertices);
+
+	int32 Index = 0;
+	for (const FGraphVertexHandle& VertexHandle : VertexHandles)
+	{
+		VertexIslandSet.Add({});
+
+		UGraphVertex* Vertex = VertexHandle.GetVertex();
+		REQUIRE(Vertex != nullptr);
+		Vertex->OnParentIslandSet.AddLambda(
+			[Index, &VertexIslandSet](const FGraphIslandHandle& IslandHandle)
+			{
+				VertexIslandSet[Index] = IslandHandle;
+			}
+		);
+
+		++Index;
+	}
+
+	for (Index = 1; Index < NumVertices; ++Index)
+	{
+		FEdgeCreationParameters Params;
+		Params.VertexHandle1 = VertexHandles[Index-1];
+		Params.VertexHandle2 = VertexHandles[Index];
+
+		TArray<FGraphEdgeHandle> Edges;
+		Graph->CreateBulkEdges({ Params }, &Edges);
+		REQUIRE(Edges.Num() == 1);
+
+		for (int32 Index2 = 0; Index2 < NumVertices; ++Index2)
+		{
+			UGraphVertex* Vertex = VertexHandles[Index2].GetVertex();
+			REQUIRE(Vertex != nullptr);
+			CHECK(Vertex->GetParentIsland() == VertexIslandSet[Index2]);
+
+			if (Index2 <= Index)
+			{
+				CHECK(VertexIslandSet[Index2].IsComplete() == true);
+			}
+			else
+			{
+				CHECK(VertexIslandSet[Index2].IsComplete() == false);
+			}
+		}
+	}
+}
\ No newline at end of file
diff --git a/Engine/Plugins/Experimental/GameplayGraph/Tests/Private/TestGraphBuilder.cpp b/Engine/Plugins/Experimental/GameplayGraph/Tests/Private/TestGraphBuilder.cpp
new file mode 100644
--- /dev/null
+++ b/Engine/Plugins/Experimental/GameplayGraph/Tests/Private/TestGraphBuilder.cpp
@@ -0,0 +1,196 @@
+// Copyright Epic Games, Inc. All Rights Reserved.
+#include "TestGraphBuilder.h"
+
+#include "TestHarness.h"
+
+FTestGraphBuilder::FTestGraphBuilder()
+{
+	Graph = NewObject<UGraph>();
+	Graph->InitializeFromProperties(FGraphProperties{ true });
+}
+
+FTestGraphBuilder::~FTestGraphBuilder()
+{
+	GraphSanityCheck();
+}
+
+void FTestGraphBuilder::PopulateVertices(uint32 Total, bool bFinalize)
+{
+	for (uint32 Count = 1; Count <= Total; ++Count)
+	{
+		FGraphUniqueIndex NodeIndex = { FGuid { 0, 0, 0, Count } };
+		FGraphVertexHandle Node = Graph->CreateVertex(NodeIndex);
+		VertexHandles.Add(Node);
+	}
+
+	if (bFinalize)
+	{
+		FinalizeVertices();
+	}
+}
+
+void FTestGraphBuilder::FinalizeVertices()
+{
+	for (uint32 Count = 0; Count < static_cast<uint32>(VertexHandles.Num()); ++Count)
+	{
+		Graph->FinalizeVertex(VertexHandles[Count]);
+	}
+}
+
+void FTestGraphBuilder::BuildFullyConnectedEdges(int32 NodesPerIsland)
+{
+	for (int32 NodeIndex = 0; NodeIndex < VertexHandles.Num(); NodeIndex += NodesPerIsland)
+	{
+		TArray<FEdgeCreationParameters> AllEdges;
+		for (int32 SourceOffset = 0; SourceOffset < NodesPerIsland; ++SourceOffset)
+		{
+			if ((NodeIndex + SourceOffset) >= VertexHandles.Num())
+			{
+				break;
+			}
+
+			for (int32 DestOffset = SourceOffset + 1; DestOffset < NodesPerIsland; ++DestOffset)
+			{
+				if ((NodeIndex + DestOffset) >= VertexHandles.Num())
+				{
+					break;
+				}
+
+				FEdgeCreationParameters Params;
+				Params.VertexHandle1 = VertexHandles[NodeIndex + SourceOffset];
+				Params.VertexHandle2 = VertexHandles[NodeIndex + DestOffset];
+				AllEdges.Add(Params);
+			}
+		}
+
+		Graph->CreateBulkEdges(MoveTemp(AllEdges));
+	}
+
+	FinalizeEdges();
+}
+
+void FTestGraphBuilder::BuildLinearEdges(int32 NodesPerIsland)
+{
+	for (int32 NodeIndex = 0; NodeIndex < VertexHandles.Num(); NodeIndex += NodesPerIsland)
+	{
+		TArray<FEdgeCreationParameters> AllEdges;
+		for (int32 SourceOffset = 1; SourceOffset < NodesPerIsland; ++SourceOffset)
+		{
+			if ((NodeIndex + SourceOffset) >= VertexHandles.Num())
+			{
+				break;
+			}
+
+			FEdgeCreationParameters Params;
+			Params.VertexHandle1 = VertexHandles[NodeIndex + SourceOffset - 1];
+			Params.VertexHandle2 = VertexHandles[NodeIndex + SourceOffset];
+			AllEdges.Add(Params);
+		}
+
+		Graph->CreateBulkEdges(MoveTemp(AllEdges));
+	}
+
+	FinalizeEdges();
+}
+
+void FTestGraphBuilder::FinalizeEdges()
+{
+	IslandHandles.Empty();
+	for (const TPair<FGraphIslandHandle, TObjectPtr<UGraphIsland>>& Pair : Graph->GetIslands())
+	{
+		IslandHandles.Add(Pair.Key);
+	}
+}
+
+TArray<FGraphEdgeHandle> FTestGraphBuilder::GetEdgesForVertex(const FGraphVertexHandle& Handle) const
+{
+	TArray<FGraphEdgeHandle> Edges;
+	if (UGraphVertex* Vertex = Handle.GetVertex())
+	{
+		Edges.Reserve(Vertex->NumEdges());
+		Vertex->ForEachAdjacentVertex(
+			[&Edges](const FGraphVertexHandle& NeighborVertexHandle, const FGraphEdgeHandle& EdgeHandle)
+			{
+				Edges.Add(EdgeHandle);
+			}
+		);
+	}
+	return Edges;
+}
+
+void FTestGraphBuilder::GraphSanityCheck()
+{
+	for (const TPair<FGraphVertexHandle, TObjectPtr<UGraphVertex>>& Data : Graph->GetVertices())
+	{
+		CHECK(Data.Key.IsComplete() == true);
+		CHECK(Data.Value != nullptr);
+		CHECK(Data.Value == Data.Key.GetVertex());
+		CHECK(Data.Key == Data.Value->Handle());
+		CHECK(Data.Key.GetUniqueIndex().IsTemporary() == false);
+		CHECK(Data.Key.GetGraph() == Graph);
+
+		if (Data.Value->GetParentIsland().IsValid())
+		{
+			CHECK(Data.Value->GetParentIsland().IsComplete());
+			CHECK(Data.Value->GetParentIsland().GetIsland() == Graph->GetIslands().FindRef(Data.Value->GetParentIsland()));
+		}
+
+		Data.Value->ForEachAdjacentVertex(
+			[this](const FGraphVertexHandle& NeighborVertexHandle, const FGraphEdgeHandle& EdgeHandle)
+			{
+				CHECK(NeighborVertexHandle.IsComplete() == true);
+				CHECK(NeighborVertexHandle.GetVertex() == Graph->GetVertices().FindRef(NeighborVertexHandle));
+
+				CHECK(EdgeHandle.IsComplete() == true);
+				CHECK(EdgeHandle.GetEdge() == Graph->GetEdges().FindRef(EdgeHandle));
+			}
+		);
+	}
+
+	for (const TPair<FGraphEdgeHandle, TObjectPtr<UGraphEdge>>& Data : Graph->GetEdges())
+	{
+		CHECK(Data.Key.IsComplete() == true);
+		CHECK(Data.Value != nullptr);
+		CHECK(Data.Value == Data.Key.GetEdge());
+		CHECK(Data.Key == Data.Value->Handle());
+		CHECK(Data.Key.GetUniqueIndex().IsTemporary() == false);
+		CHECK(Data.Key.GetGraph() == Graph);
+
+		CHECK(Data.Value->NodeA().IsComplete() == true);
+		CHECK(Data.Value->NodeA().GetVertex() == Graph->GetVertices().FindRef(Data.Value->NodeA()));
+
+		CHECK(Data.Value->NodeB().IsComplete() == true);
+		CHECK(Data.Value->NodeB().GetVertex() == Graph->GetVertices().FindRef(Data.Value->NodeB()));
+	}
+
+	for (const TPair<FGraphIslandHandle, TObjectPtr<UGraphIsland>>& Data : Graph->GetIslands())
+	{
+		CHECK(Data.Key.IsComplete() == true);
+		CHECK(Data.Value != nullptr);
+		CHECK(Data.Value == Data.Key.GetIsland());
+		CHECK(Data.Key == Data.Value->Handle());
+		CHECK(Data.Key.GetUniqueIndex().IsTemporary() == false);
+		CHECK(Data.Key.GetGraph() == Graph);
+
+		for (const FGraphVertexHandle& VertexHandle : Data.Value->GetVertices())
+		{
+			CHECK(VertexHandle.IsComplete() == true);
+			CHECK(VertexHandle.GetVertex() == Graph->GetVertices().FindRef(VertexHandle));
+		}
+
+		IslandVertexParentIslandSanityCheck(Data.Key);
+	}
+}
+
+void FTestGraphBuilder::IslandVertexParentIslandSanityCheck(const FGraphIslandHandle& IslandHandle)
+{
+	UGraphIsland* Island = IslandHandle.GetIsland();
+	REQUIRE(Island != nullptr);
+
+	Island->ForEachVertex(
+		[&IslandHandle](const FGraphVertexHandle& VertexHandle)
+		{
+			CHECK(VertexHandle.GetVertex()->GetParentIsland() == IslandHandle);
+		}
+	);
+}
diff --git a/Engine/Plugins/Experimental/GameplayGraph/Tests/Private/TestGraphBuilder.h b/Engine/Plugins/Experimental/GameplayGraph/Tests/Private/TestGraphBuilder.h
new file mode 100644
--- /dev/null
+++ b/Engine/Plugins/Experimental/GameplayGraph/Tests/Private/TestGraphBuilder.h
@@ -0,0 +1,29 @@
+// Copyright Epic Games, Inc. All Rights Reserved.
+#pragma once
+
+#include "Graph/Graph.h"
+
+class FTestGraphBuilder
+{
+public:
+	FTestGraphBuilder();
+	~FTestGraphBuilder();
+
+	void PopulateVertices(uint32 Total, bool bFinalize);
+	void FinalizeVertices();
+
+	void BuildFullyConnectedEdges(int32 NodesPerIsland);
+	void BuildLinearEdges(int32 NodesPerIsland);
+
+	void FinalizeEdges();
+
+	TArray<FGraphEdgeHandle> GetEdgesForVertex(const FGraphVertexHandle& Handle) const;
+
+protected:
+	TObjectPtr<UGraph> Graph;
+	TArray<FGraphVertexHandle> VertexHandles;
+	TArray<FGraphIslandHandle> IslandHandles;
+
+	void GraphSanityCheck();
+	void IslandVertexParentIslandSanityCheck(const FGraphIslandHandle& IslandHandle);
+};
\ No newline at end of file
diff --git a/Engine/Plugins/Experimental/GameplayGraph/Tests/Private/TestGraphElement.cpp b/Engine/Plugins/Experimental/GameplayGraph/Tests/Private/TestGraphElement.cpp
new file mode 100644
--- /dev/null
+++ b/Engine/Plugins/Experimental/GameplayGraph/Tests/Private/TestGraphElement.cpp
@@ -0,0 +1,4 @@
+// Copyright Epic Games, Inc. All Rights Reserved.
+#include "TestGraphElement.h"
+
+#include UE_INLINE_GENERATED_CPP_BY_NAME(TestGraphElement)
diff --git a/Engine/Plugins/Experimental/GameplayGraph/Tests/Private/TestGraphElement.h b/Engine/Plugins/Experimental/GameplayGraph/Tests/Private/TestGraphElement.h
new file mode 100644
--- /dev/null
+++ b/Engine/Plugins/Experimental/GameplayGraph/Tests/Private/TestGraphElement.h
@@ -0,0 +1,34 @@
+// Copyright Epic Games, Inc. All Rights Reserved.
+#pragma once
+
+#include "Graph/GraphElement.h"
+
+#include "TestGraphElement.generated.h"
+
+UCLASS()
+class UTestGraphElement : public UGraphElement
+{
+	GENERATED_BODY()
+
+public:
+	void PublicSetUniqueIndex(FGraphUniqueIndex InUniqueIndex)
+	{
+		SetUniqueIndex(InUniqueIndex);
+	}
+
+	FGraphUniqueIndex PublicGetUniqueIndex()
+	{
+		return GetUniqueIndex();
+	}
+
+
+	void PublicSetParentGraph(TObjectPtr<UGraph> InGraph)
+	{
+		SetParentGraph(InGraph);
+	}
+
+	TObjectPtr<UGraph> PublicGetGraph() const
+	{
+		return GetGraph();
+	}
+};
diff --git a/Engine/Plugins/Experimental/GameplayGraph/Tests/Private/TestGroupEvents.cpp b/Engine/Plugins/Experimental/GameplayGraph/Tests/Private/TestGroupEvents.cpp
new file mode 100644
--- /dev/null
+++ b/Engine/Plugins/Experimental/GameplayGraph/Tests/Private/TestGroupEvents.cpp
@@ -0,0 +1,23 @@
+// Copyright Epic Games, Inc. All Rights Reserved.
+#include "TestCommon/CoreUtilities.h"
+#include "TestCommon/CoreUObjectUtilities.h"
+
+#include <catch2/catch_test_macros.hpp>
+
+GROUP_BEFORE_GLOBAL(Catch::DefaultGroup)
+{
+	InitTaskGraph();
+	InitThreadPool(true);
+#if WITH_COREUOBJECT
+	InitCoreUObject();
+#endif
+}
+
+GROUP_AFTER_GLOBAL(Catch::DefaultGroup)
+{
+#if WITH_COREUOBJECT
+	CleanupCoreUObject();
+#endif
+	CleanupThreadPool();
+	CleanupTaskGraph();
+}
