remove explicit edges from the gameplay graph to significantly cut memory usage

diff --git a/Engine/Plugins/Experimental/GameplayGraph/Source/GameplayGraph/Private/Graph/Graph.cpp b/Engine/Plugins/Experimental/GameplayGraph/Source/GameplayGraph/Private/Graph/Graph.cpp
--- a/Engine/Plugins/Experimental/GameplayGraph/Source/GameplayGraph/Private/Graph/Graph.cpp
+++ b/Engine/Plugins/Experimental/GameplayGraph/Source/GameplayGraph/Private/Graph/Graph.cpp
@@ -2,7 +2,6 @@
 #include "Graph/Graph.h"
 
 #include "Graph/Algorithms/Connectivity/ConnectedComponents.h"
-#include "Graph/GraphEdge.h"
 #include "Graph/GraphVertex.h"
 #include "Graph/GraphSerialization.h"
 #include "GenericPlatform/GenericPlatformMath.h"
@@ -10,17 +9,25 @@
 
 #include UE_INLINE_GENERATED_CPP_BY_NAME(Graph)
 
-
 namespace Graph
 {
 	static bool bFixupGraphOnLoad = true;
 	static FAutoConsoleVariableRef FixupGraphOnLoadCVar(TEXT("GameplayGraph.FixupGraphOnLoad"), bFixupGraphOnLoad, TEXT("Merge and split islands when loaded from serialization."));
 }
 
+FEdgeSpecifier::FEdgeSpecifier(const FGraphVertexHandle& InVertexHandle1, const FGraphVertexHandle& InVertexHandle2)
+: VertexHandle1(InVertexHandle1)
+, VertexHandle2(InVertexHandle2)
+{
+	if (VertexHandle2 < VertexHandle1)
+	{
+		std::swap(VertexHandle1, VertexHandle2);
+	}
+}
+
 void UGraph::Empty()
 {
 	Vertices.Empty();
-	Edges.Empty();
 	Islands.Empty();
 }
 
@@ -70,18 +77,12 @@ void UGraph::RegisterVertex(TObjectPtr<UGraphVertex> Vertex)
 	Vertices.Add(Handle, Vertex);
 }
 
-FGraphEdgeHandle UGraph::CreateEdge(FGraphVertexHandle Node1, FGraphVertexHandle Node2, FGraphUniqueIndex InUniqueIndex, bool bMergeIslands)
+bool UGraph::CreateEdge(FGraphVertexHandle Node1, FGraphVertexHandle Node2, bool bMergeIslands)
 {
 	TRACE_CPUPROFILER_EVENT_SCOPE(UGraph::CreateEdge);
 	if (!Node1.IsComplete() || !Node2.IsComplete())
 	{
-		return {};
-	}
-
-	TObjectPtr<UGraphEdge> Edge = CreateTypedEdge();
-	if (!Edge)
-	{
-		return {};
+		return false;
 	}
 
 	if (Node2 < Node1)
@@ -89,22 +90,17 @@ FGraphEdgeHandle UGraph::CreateEdge(FGraphVertexHandle Node1, FGraphVertexHandle
 		std::swap(Node1, Node2);
 	}
 
-	TObjectPtr<UGraphVertex> Node1Ptr = Node1.GetVertex();
-	TObjectPtr<UGraphVertex> Node2Ptr = Node2.GetVertex();
+	UGraphVertex* Node1Ptr = Node1.GetVertex();
+	UGraphVertex* Node2Ptr = Node2.GetVertex();
 	if (!Node1Ptr || !Node2Ptr)
 	{
-		return {};
+		return false;
 	}
 
 	// For edges, we also need to make sure the edge doesn't already exist.
 	if (Node1Ptr->HasEdgeTo(Node2) || Node2Ptr->HasEdgeTo(Node1))
 	{
-		return {};
-	}
-
-	if (!ensure(InUniqueIndex.IsValid()))
-	{
-		return {};
+		return false;
 	}
 
 	if (Properties.bGenerateIslands)
@@ -112,8 +108,8 @@ FGraphEdgeHandle UGraph::CreateEdge(FGraphVertexHandle Node1, FGraphVertexHandle
 		const FGraphIslandHandle& Island1 = Node1Ptr->GetParentIsland();
 		const FGraphIslandHandle& Island2 = Node2Ptr->GetParentIsland();
 
-		const TObjectPtr<UGraphIsland> Island1Ptr = Island1.GetIsland();
-		const TObjectPtr<UGraphIsland> Island2Ptr = Island2.GetIsland();
+		const UGraphIsland* Island1Ptr = Island1.GetIsland();
+		const UGraphIsland* Island2Ptr = Island2.GetIsland();
 
 		if (Island1Ptr && Island2Ptr)
 		{
@@ -124,7 +120,7 @@ FGraphEdgeHandle UGraph::CreateEdge(FGraphVertexHandle Node1, FGraphVertexHandle
 				return {};
 			}
 		}
-		else if (TObjectPtr<UGraphIsland> RelevantIsland = Island1Ptr ? Island1Ptr : Island2Ptr)
+		else if (const UGraphIsland* RelevantIsland = Island1Ptr ? Island1Ptr : Island2Ptr)
 		{
 			// Regular add scenario.
 			if (!RelevantIsland->IsOperationAllowed(EGraphIslandOperations::Add))
@@ -134,17 +130,9 @@ FGraphEdgeHandle UGraph::CreateEdge(FGraphVertexHandle Node1, FGraphVertexHandle
 		}
 	}
 
-	Edge->OnCreate();
-	Edge->SetUniqueIndex(InUniqueIndex);
-	Edge->SetNodes(Node1, Node2);
-
-	RegisterEdge(Edge);
-
 	// Also need to make sure the nodes are aware of the new edge.
-	Node1Ptr->AddEdgeTo(Node2, Edge->Handle());
-	Node2Ptr->AddEdgeTo(Node1, Edge->Handle());
-
-	OnEdgeCreated.Broadcast(Edge->Handle());
+	Node1Ptr->AddEdgeTo(Node2);
+	Node2Ptr->AddEdgeTo(Node1);
 
 	// If we want to keep track of islands, this is where we need to create/merge islands.
 	if (Properties.bGenerateIslands)
@@ -169,48 +157,31 @@ FGraphEdgeHandle UGraph::CreateEdge(FGraphVertexHandle Node1, FGraphVertexHandle
 
 		if (bMergeIslands)
 		{
-			MergeOrCreateIslands({ Edge->Handle() });
+			const FEdgeSpecifier Edge{ Node1, Node2 };
+			MergeOrCreateIslands({ Edge });
 		}
 	}
 
-	return Edge->Handle();
+	return true;
 }
 
-void UGraph::CreateBulkEdges(TArray<FEdgeCreationParameters>&& NodesToConnect, TArray<FGraphEdgeHandle>* OutEdges)
+void UGraph::CreateBulkEdges(TArray<FEdgeSpecifier>&& NodesToConnect)
 {
-	// Create all edges normally but don't call MergeOrCreateIslands yet. We'll use the bulk function instead.
-	TArray<FGraphEdgeHandle> NewEdges;
-	NewEdges.Reserve(NodesToConnect.Num());
+	TArray<FEdgeSpecifier> FinalEdges;
+	FinalEdges.Reserve(NodesToConnect.Num());
 
-	for (const FEdgeCreationParameters& Params : NodesToConnect)
+	// Create all edges normally but don't call MergeOrCreateIslands yet. We'll use the bulk function instead.
+	for (const FEdgeSpecifier& Params : NodesToConnect)
 	{
-		if (FGraphEdgeHandle Edge = CreateEdge(Params.VertexHandle1, Params.VertexHandle2, Params.EdgeIndex, false); Edge.IsValid())
+		if (CreateEdge(Params.GetVertexHandle1(), Params.GetVertexHandle2(), false))
 		{
-			NewEdges.Add(Edge);
+			FinalEdges.Add(Params);
 		}
 	}
 
 	// MergeOrCreateIslands incrementally determines island connectivity one edge at a time. This is efficient if we're handling a single edge but less
 	// efficient if we're trying to add a bunch of edges all at the same time since it'll cause a vertex to jump between islands. The Bulk function helps prevent that.
-	MergeOrCreateIslands(NewEdges);
-
-	if (OutEdges)
-	{
-		*OutEdges = MoveTemp(NewEdges);
-	}
-}
-
-void UGraph::RegisterEdge(TObjectPtr<UGraphEdge> Edge)
-{
-	TRACE_CPUPROFILER_EVENT_SCOPE(UGraph::RegisterEdge);
-	if (!Edge)
-	{
-		return;
-	}
-
-	Edge->SetParentGraph(this);
-	const FGraphEdgeHandle Handle = Edge->Handle();
-	Edges.Add(Handle, Edge);
+	MergeOrCreateIslands(FinalEdges);
 }
 
 FGraphIslandHandle UGraph::CreateIsland(TArray<FGraphVertexHandle> InputNodes, FGraphUniqueIndex InUniqueIndex)
@@ -271,10 +242,10 @@ void UGraph::RemoveIsland(const FGraphIslandHandle& IslandHandle)
 	Islands.Remove(IslandHandle);
 }
 
-void UGraph::MergeOrCreateIslands(const TArray<FGraphEdgeHandle>& InEdges)
+void UGraph::MergeOrCreateIslands(const TArray<FEdgeSpecifier>& InEdges)
 {
 	TRACE_CPUPROFILER_EVENT_SCOPE(UGraph::MergeOrCreateIslands);
-	if (Edges.IsEmpty())
+	if (InEdges.IsEmpty())
 	{
 		return;
 	}
@@ -356,19 +327,13 @@ void UGraph::MergeOrCreateIslands(const TArray<FGraphEdgeHandle>& InEdges)
 		}
 	};
 
-	for (const FGraphEdgeHandle& EdgeHandle : InEdges)
+	for (const FEdgeSpecifier& EdgeHandle : InEdges)
 	{
-		if (!EdgeHandle.IsComplete())
-		{
-			continue;
-		}
+		const FGraphVertexHandle& AHandle = EdgeHandle.GetVertexHandle1();
+		const FGraphVertexHandle& BHandle = EdgeHandle.GetVertexHandle2();
 
-		TObjectPtr<UGraphEdge> Edge = EdgeHandle.GetEdge();
-		const FGraphVertexHandle& AHandle = Edge->NodeA();
-		const FGraphVertexHandle& BHandle = Edge->NodeB();
-
-		TObjectPtr<UGraphVertex> NodeA = AHandle.GetVertex();
-		TObjectPtr<UGraphVertex> NodeB = BHandle.GetVertex();
+		UGraphVertex* NodeA = AHandle.GetVertex();
+		UGraphVertex* NodeB = BHandle.GetVertex();
 		if (!NodeA || !NodeB)
 		{
 			return;
@@ -507,27 +472,27 @@ void UGraph::RemoveBulkVertices(const TArray<FGraphVertexHandle>& InHandles)
 			if (TObjectPtr<UGraphIsland> Island = Node->GetParentIsland().GetIsland())
 			{
 				AffectedIslands.Add(Node->GetParentIsland());
-				Island->RemoveVertex(NodeHandle);
+				Island->RemoveVertex(Node->Handle());
 			}
 
 			// We must remove every edge this node is a part of. Need to make a copy of the edges because
 			// otherwise we're modifying the container during iteration over it in ForEachAdjacentVertex.
-			TArray<FGraphEdgeHandle> EdgeCopy;
+			TArray<FGraphVertexHandle> EdgeCopy;
 			EdgeCopy.Reserve(Node->NumEdges());
 			Node->ForEachAdjacentVertex(
-				[&EdgeCopy](const FGraphVertexHandle& OtherNodeHandle, const FGraphEdgeHandle& EdgeHandle)
+				[&EdgeCopy](const FGraphVertexHandle& OtherNodeHandle)
 				{
-					EdgeCopy.Add(EdgeHandle);
+					EdgeCopy.Add(OtherNodeHandle);
 				}
 			);
 
-			for (const FGraphEdgeHandle& EdgeHandle : EdgeCopy)
+			for (const FGraphVertexHandle& AdjacentVertexHandle : EdgeCopy)
 			{
 				// Don't immediately handle islands. We'll do it later.
-				RemoveEdgeInternal(EdgeHandle, false);
+				RemoveEdgeInternal(Node->Handle(), AdjacentVertexHandle, false);
 			}
 
-			AffectedVertices.Add(NodeHandle);
+			AffectedVertices.Add(Node->Handle());
 		}
 	}
 
@@ -551,48 +516,24 @@ void UGraph::RemoveBulkVertices(const TArray<FGraphVertexHandle>& InHandles)
 	}
 }
 
-void UGraph::RemoveEdge(const FGraphEdgeHandle& EdgeHandle)
-{
-	RemoveEdgeInternal(EdgeHandle, true);
-}
-
-void UGraph::RemoveEdgeInternal(const FGraphEdgeHandle& EdgeHandle, bool bHandleIslands)
+void UGraph::RemoveEdgeInternal(const FGraphVertexHandle& VertexHandleA, const FGraphVertexHandle& VertexHandleB, bool bHandleIslands)
 {
 	TRACE_CPUPROFILER_EVENT_SCOPE(UGraph::RemoveEdgeInternal);
-	if (!EdgeHandle.IsValid())
-	{
-		return;
-	}
-
-	TObjectPtr<UGraphEdge> Edge = EdgeHandle.GetEdge();
-	if (!Edge)
-	{
-		return;
-	}
-
 	FGraphIslandHandle IslandHandle;
 
 	// Need to remove the edge reference from both nodes.
-	if (const FGraphVertexHandle& NodeHandleA = Edge->NodeA(); NodeHandleA.IsValid())
+	if (UGraphVertex* NodeA = VertexHandleA.GetVertex())
 	{
-		if (TObjectPtr<UGraphVertex> Node = NodeHandleA.GetVertex(); Node)
-		{
-			Node->RemoveEdge(EdgeHandle);
-			IslandHandle = Node->GetParentIsland();
-		}
+		NodeA->RemoveEdge(VertexHandleB);
+		IslandHandle = NodeA->GetParentIsland();
 	}
 
-	if (const FGraphVertexHandle& NodeHandleB = Edge->NodeB(); NodeHandleB.IsValid())
+	if (UGraphVertex* NodeB = VertexHandleB.GetVertex())
 	{
-		if (TObjectPtr<UGraphVertex> Node = NodeHandleB.GetVertex(); Node)
+		NodeB->RemoveEdge(VertexHandleA);
+		if (!IslandHandle.IsValid())
 		{
-			Node->RemoveEdge(EdgeHandle);
-
-			// Technically shouldn't be necessary but in here just in case.
-			if (!IslandHandle.IsValid())
-			{
-				IslandHandle = Node->GetParentIsland();
-			}
+			IslandHandle = NodeB->GetParentIsland();
 		}
 	}
 	
@@ -602,8 +543,6 @@ void UGraph::RemoveEdgeInternal(const FGraphEdgeHandle& EdgeHandle, bool bHandle
 		// Note that we can assume that both nodes are in the same island.
 		RemoveOrSplitIsland(IslandHandle.GetIsland());
 	}
-
-	Edges.Remove(EdgeHandle);
 }
 
 void UGraph::RemoveOrSplitIsland(TObjectPtr<UGraphIsland> Island)
@@ -689,11 +628,6 @@ TObjectPtr<UGraphVertex> UGraph::CreateTypedVertex() const
 	return NewObject<UGraphVertex>(const_cast<UGraph*>(this), UGraphVertex::StaticClass());
 }
 
-TObjectPtr<UGraphEdge> UGraph::CreateTypedEdge() const
-{
-	return NewObject<UGraphEdge>(const_cast<UGraph*>(this), UGraphEdge::StaticClass());
-}
-
 TObjectPtr<UGraphIsland> UGraph::CreateTypedIsland() const
 {
 	return NewObject<UGraphIsland>(const_cast<UGraph*>(this), UGraphIsland::StaticClass());
@@ -722,16 +656,16 @@ void UGraph::FinalizeVertex(const FGraphVertexHandle& InHandle)
 
 	// Need to iterate through edges to see if any edge was added prior to vertex finalization.
 	// If so, we need to do a merge.
-	TArray<FGraphEdgeHandle> EdgesToMerge;
+	TArray<FEdgeSpecifier> EdgesToMerge;
 	Node->ForEachAdjacentVertex(
-		[&IslandHandle, &EdgesToMerge](const FGraphVertexHandle& NeighborVertexHandle, const FGraphEdgeHandle& EdgeHandle)
+		[&InHandle, &IslandHandle, &EdgesToMerge](const FGraphVertexHandle& NeighborVertexHandle)
 		{
 			if (UGraphVertex* NeighborVertex = NeighborVertexHandle.GetVertex())
 			{
 				const FGraphIslandHandle& NeighborIslandHandle = NeighborVertex->GetParentIsland();
 				if (NeighborIslandHandle.IsComplete() && NeighborIslandHandle != IslandHandle)
 				{
-					EdgesToMerge.Add(EdgeHandle);
+					EdgesToMerge.Add(FEdgeSpecifier{InHandle, NeighborVertexHandle});
 				}
 			}
 		}
@@ -759,11 +693,6 @@ UGraphVertex* UGraph::GetSafeVertexFromHandle(const FGraphVertexHandle& Handle)
 	return Handle.IsComplete() ? Handle.GetVertex() : Vertices.FindRef(Handle).Get();
 }
 
-UGraphEdge* UGraph::GetSafeEdgeFromHandle(const FGraphEdgeHandle& Handle) const
-{
-	return Handle.IsComplete() ? Handle.GetEdge() : Edges.FindRef(Handle).Get();
-}
-
 UGraphIsland* UGraph::GetSafeIslandFromHandle(const FGraphIslandHandle& Handle) const
 {
 	return Handle.IsComplete() ? Handle.GetIsland() : Islands.FindRef(Handle).Get();
@@ -774,11 +703,6 @@ void UGraph::ReserveVertices(int32 Delta)
 	Vertices.Reserve(Vertices.Num() + Delta);
 }
 
-void UGraph::ReserveEdges(int32 Delta)
-{
-	Edges.Reserve(Edges.Num() + Delta);
-}
-
 void UGraph::ReserveIslands(int32 Delta)
 {
 	Islands.Reserve(Islands.Num() + Delta);
@@ -787,11 +711,16 @@ void UGraph::ReserveIslands(int32 Delta)
 void operator<<(IGraphSerialization& Output, const UGraph& Graph)
 {
 	TRACE_CPUPROFILER_EVENT_SCOPE_STR("Gameplay Graph Serialization");
-	Output.Initialize(Graph.NumVertices(), Graph.NumEdges(), Graph.NumIslands());
+
+	// We don't actually have a good estimate of the number of edges anymore - so just guesstimating number of vertices.
+	Output.Initialize(Graph.NumVertices(), Graph.NumVertices(), Graph.NumIslands());
 	Output.WriteGraphProperties(Graph.GetProperties());
 
 	{
-		TRACE_CPUPROFILER_EVENT_SCOPE_STR("Serialize Vertices");
+		TRACE_CPUPROFILER_EVENT_SCOPE_STR("Serialize Vertices and Edges");
+		TSet<FEdgeSpecifier> SeenEdges;
+		SeenEdges.Reserve(Graph.NumVertices());
+
 		for (const TPair<FGraphVertexHandle, TObjectPtr<UGraphVertex>>& Kvp : Graph.GetVertices())
 		{
 			if (!ensure(Kvp.Key.IsComplete() && Kvp.Value))
@@ -800,19 +729,18 @@ void operator<<(IGraphSerialization& Output, const UGraph& Graph)
 			}
 
 			Output.WriteGraphVertex(Kvp.Key, Kvp.Value);
-		}
-	}
 
-	{
-		TRACE_CPUPROFILER_EVENT_SCOPE_STR("Serialize Edges");
-		for (const TPair<FGraphEdgeHandle, TObjectPtr<UGraphEdge>>& Kvp : Graph.GetEdges())
-		{
-			if (!ensure(Kvp.Key.IsComplete() && Kvp.Value))
-			{
-				continue;
-			}
-
-			Output.WriteGraphEdge(Kvp.Key, Kvp.Value);
+			Kvp.Value->ForEachAdjacentVertex(
+				[&SeenEdges, &Kvp, &Output](const FGraphVertexHandle& AdjacentVertexHandle)
+				{
+					FEdgeSpecifier Edge{ Kvp.Key, AdjacentVertexHandle };
+					if (!SeenEdges.Contains(Edge))
+					{
+						SeenEdges.Add(Edge);
+						Output.WriteGraphEdge(Edge.GetVertexHandle1(), Edge.GetVertexHandle2());
+					}
+				}
+			);
 		}
 	}
 
@@ -841,7 +769,7 @@ void operator>>(const IGraphDeserialization& Input, UGraph& Graph)
 		Input.ForEveryVertex(
 			[&Graph](const FGraphVertexHandle& InHandle)
 			{
-				if (!ensure(InHandle.IsValid()))
+				if (!InHandle.IsValid())
 				{
 					return FGraphVertexHandle{};
 				}
@@ -851,27 +779,27 @@ void operator>>(const IGraphDeserialization& Input, UGraph& Graph)
 		);
 	}
 
-	TArray<FGraphEdgeHandle> AllEdges;
+	TArray<FEdgeSpecifier> AllEdges;
 	AllEdges.Reserve(Input.NumEdges());
 	{
 		TRACE_CPUPROFILER_EVENT_SCOPE_STR("Deserialize Edges");
-		Graph.ReserveEdges(Input.NumEdges());
 		Input.ForEveryEdge(
-			[&Graph, &AllEdges](const FGraphEdgeHandle& InHandle, const IGraphDeserialization::FEdgeConstructionData& Data)
+			[&Graph, &AllEdges](const FEdgeSpecifier& Data)
 			{
-				if (!ensure(InHandle.IsValid()))
-				{
-					return FGraphEdgeHandle{};
-				}
+				const FGraphVertexHandle CompleteHandle1 = Graph.GetCompleteNodeHandle(Data.GetVertexHandle1());
+				const FGraphVertexHandle CompleteHandle2 = Graph.GetCompleteNodeHandle(Data.GetVertexHandle2());
 
-				FGraphEdgeHandle Handle = Graph.CreateEdge(
-					Graph.GetCompleteNodeHandle(Data.Vertex1),
-					Graph.GetCompleteNodeHandle(Data.Vertex2),
-					InHandle.GetUniqueIndex(),
+				const bool bSuccess = Graph.CreateEdge(
+					CompleteHandle1,
+					CompleteHandle2,
 					false
 				);
-				AllEdges.Add(Handle);
-				return Handle;
+
+				if (bSuccess)
+				{
+					AllEdges.Add(FEdgeSpecifier{CompleteHandle1, CompleteHandle2});
+				}
+				return bSuccess;
 			}
 		);
 	}
@@ -879,14 +807,13 @@ void operator>>(const IGraphDeserialization& Input, UGraph& Graph)
 	TArray<FGraphIslandHandle> AllIslands;
 	AllIslands.Reserve(Input.NumIslands());
 
-
 	{
 		TRACE_CPUPROFILER_EVENT_SCOPE_STR("Deserialize Islands");
 		Graph.ReserveIslands(Input.NumIslands());
 		Input.ForEveryIsland(
 			[&Graph, &AllIslands](const FGraphIslandHandle& InHandle, const IGraphDeserialization::FIslandConstructionData& Data)
 			{
-				if (!ensure(InHandle.IsValid()))
+				if (!InHandle.IsValid())
 				{
 					return FGraphIslandHandle{};
 				}
diff --git a/Engine/Plugins/Experimental/GameplayGraph/Source/GameplayGraph/Private/Graph/GraphEdge.cpp b/Engine/Plugins/Experimental/GameplayGraph/Source/GameplayGraph/Private/Graph/GraphEdge.cpp
deleted file mode 100644
diff --git a/Engine/Plugins/Experimental/GameplayGraph/Source/GameplayGraph/Private/Graph/GraphHandle.cpp b/Engine/Plugins/Experimental/GameplayGraph/Source/GameplayGraph/Private/Graph/GraphHandle.cpp
--- a/Engine/Plugins/Experimental/GameplayGraph/Source/GameplayGraph/Private/Graph/GraphHandle.cpp
+++ b/Engine/Plugins/Experimental/GameplayGraph/Source/GameplayGraph/Private/Graph/GraphHandle.cpp
@@ -4,7 +4,6 @@
 
 #include "Graph/Graph.h"
 #include "Graph/GraphElement.h"
-#include "Graph/GraphEdge.h"
 #include "Graph/GraphIsland.h"
 #include "Graph/GraphVertex.h"
 #include "Misc/Guid.h"
@@ -14,7 +13,6 @@
 DEFINE_LOG_CATEGORY(LogGameplayGraph)
 
 FGraphVertexHandle FGraphVertexHandle::Invalid;
-FGraphEdgeHandle FGraphEdgeHandle::Invalid;
 FGraphIslandHandle FGraphIslandHandle::Invalid;
 
 FGraphHandle::FGraphHandle(FGraphUniqueIndex InUniqueIndex, UGraph* InGraph)
@@ -79,26 +77,6 @@ bool FGraphVertexHandle::HasElement() const
 	return false;
 }
 
-UGraphEdge* FGraphEdgeHandle::GetEdge() const
-{
-	if (UGraph* Graph = GetGraph())
-	{
-		return Graph->GetEdges().FindRef(*this);
-	}
-
-	return nullptr;
-}
-
-bool FGraphEdgeHandle::HasElement() const
-{
-	if (UGraph* Graph = GetGraph())
-	{
-		return Graph->GetEdges().Contains(*this);
-	}
-
-	return false;
-}
-
 UGraphIsland* FGraphIslandHandle::GetIsland() const
 {
 	if (UGraph* Graph = GetGraph())
diff --git a/Engine/Plugins/Experimental/GameplayGraph/Source/GameplayGraph/Private/Graph/GraphVertex.cpp b/Engine/Plugins/Experimental/GameplayGraph/Source/GameplayGraph/Private/Graph/GraphVertex.cpp
--- a/Engine/Plugins/Experimental/GameplayGraph/Source/GameplayGraph/Private/Graph/GraphVertex.cpp
+++ b/Engine/Plugins/Experimental/GameplayGraph/Source/GameplayGraph/Private/Graph/GraphVertex.cpp
@@ -1,7 +1,6 @@
 // Copyright Epic Games, Inc. All Rights Reserved.
 
 #include "Graph/GraphVertex.h"
-#include "Graph/GraphEdge.h"
 
 #include UE_INLINE_GENERATED_CPP_BY_NAME(GraphVertex)
 
@@ -15,30 +14,20 @@ bool UGraphVertex::HasEdgeTo(const FGraphVertexHandle& Other) const
 	return Edges.Contains(Other);
 }
 
-const FGraphEdgeHandle* UGraphVertex::FindEdgeTo(const FGraphVertexHandle& Other) const
+void UGraphVertex::AddEdgeTo(const FGraphVertexHandle& Node)
 {
-	return Edges.Find(Other);
+	Edges.Add(Node);
 }
 
-void UGraphVertex::AddEdgeTo(const FGraphVertexHandle& Node, const FGraphEdgeHandle& Edge)
+void UGraphVertex::RemoveEdge(const FGraphVertexHandle& AdjacentVertexHandle)
 {
-	Edges.Add(Node, Edge);
-}
-
-void UGraphVertex::RemoveEdge(const FGraphEdgeHandle& EdgeHandle)
-{
-	if (!EdgeHandle.IsValid())
-	{
-		return;
-	}
-
-	TObjectPtr<UGraphEdge> Edge = EdgeHandle.GetEdge();
-	if (!Edge || !Edge->ContainsNode(Handle()))
+	UGraphVertex* AdjacentVertex = AdjacentVertexHandle.GetVertex();
+	if (ensure(AdjacentVertex))
 	{
-		return;
+		AdjacentVertex->Edges.Remove(Handle());
 	}
 
-	Edges.Remove(Edge->GetOtherNode(Handle()));
+	Edges.Remove(AdjacentVertexHandle);
 }
 
 void UGraphVertex::HandleOnVertexRemoved()
diff --git a/Engine/Plugins/Experimental/GameplayGraph/Source/GameplayGraph/Public/Graph/Graph.h b/Engine/Plugins/Experimental/GameplayGraph/Source/GameplayGraph/Public/Graph/Graph.h
--- a/Engine/Plugins/Experimental/GameplayGraph/Source/GameplayGraph/Public/Graph/Graph.h
+++ b/Engine/Plugins/Experimental/GameplayGraph/Source/GameplayGraph/Public/Graph/Graph.h
@@ -6,7 +6,6 @@
 #include "Containers/Array.h"
 #include "Containers/Map.h"
 #include "Containers/Set.h"
-#include "Graph/GraphEdge.h"
 #include "Graph/GraphHandle.h"
 #include "Graph/GraphIsland.h"
 #include "Graph/GraphVertex.h"
@@ -31,14 +30,27 @@ struct FGraphProperties
 };
 
 DECLARE_MULTICAST_DELEGATE_OneParam(FOnGraphVertexCreated, const FGraphVertexHandle&);
-DECLARE_MULTICAST_DELEGATE_OneParam(FOnGraphEdgeCreated, const FGraphEdgeHandle&);
 DECLARE_MULTICAST_DELEGATE_OneParam(FOnGraphIslandCreated, const FGraphIslandHandle&);
 
-struct FEdgeCreationParameters
+struct GAMEPLAYGRAPH_API FEdgeSpecifier
 {
+	FEdgeSpecifier(const FGraphVertexHandle& InVertexHandle1, const FGraphVertexHandle& InVertexHandle2);
+
+	// Comparison operators
+	friend bool operator==(const FEdgeSpecifier& Lhs, const FEdgeSpecifier& Rhs) = default;
+	friend bool operator!=(const FEdgeSpecifier& Lhs, const FEdgeSpecifier& Rhs) = default;
+
+	friend uint32 GAMEPLAYGRAPH_API GetTypeHash(const FEdgeSpecifier& InEdge)
+	{
+		return GetTypeHash(TPair<FGraphVertexHandle, FGraphVertexHandle>{InEdge.VertexHandle1, InEdge.VertexHandle2});
+	}
+
+	const FGraphVertexHandle& GetVertexHandle1() const { return VertexHandle1; }
+	const FGraphVertexHandle& GetVertexHandle2() const { return VertexHandle2; }
+
+private:
 	FGraphVertexHandle VertexHandle1;
 	FGraphVertexHandle VertexHandle2;
-	FGraphUniqueIndex EdgeIndex = FGraphUniqueIndex::CreateUniqueIndex();
 };
 
 /**
@@ -80,16 +92,13 @@ public:
 	FGraphVertexHandle CreateVertex(FGraphUniqueIndex InUniqueIndex = FGraphUniqueIndex::CreateUniqueIndex(false));
 
 	/** Creates edges in bulk. This is more efficient than calling CreateEdge multiple times since we will only try to assign a node to an island once. */
-	void CreateBulkEdges(TArray<FEdgeCreationParameters>&& NodesToConnect, TArray<FGraphEdgeHandle>* OutEdges = nullptr);
+	void CreateBulkEdges(TArray<FEdgeSpecifier>&& NodesToConnect);
 
 	/** Removes a node from the graph along with any edges that contain it. */
 	void RemoveVertex(const FGraphVertexHandle& NodeHandle);
 
 	void RemoveBulkVertices(const TArray<FGraphVertexHandle>& InHandles);
 
-	/** Removes an edge from the graph. */
-	void RemoveEdge(const FGraphEdgeHandle& EdgeHandle);
-
 	/** This should be called immediately after a node and any relevant edges have been added to the graph. */
 	virtual void FinalizeVertex(const FGraphVertexHandle& InHandle);
 
@@ -100,7 +109,6 @@ public:
 	void RefreshIslandConnectivity(const FGraphIslandHandle& IslandHandle);
 
 	int32 NumVertices() const { return Vertices.Num(); }
-	int32 NumEdges() const { return Edges.Num(); }
 	int32 NumIslands() const { return Islands.Num(); }
 
 	template<typename TLambda>
@@ -113,16 +121,13 @@ public:
 	}
 
 	FOnGraphVertexCreated OnVertexCreated;
-	FOnGraphEdgeCreated OnEdgeCreated;
 	FOnGraphIslandCreated OnIslandCreated;
 
 	const TMap<FGraphVertexHandle, TObjectPtr<UGraphVertex>>& GetVertices() const { return Vertices; }
-	const TMap<FGraphEdgeHandle, TObjectPtr<UGraphEdge>>& GetEdges() const { return Edges; }
 	const TMap<FGraphIslandHandle, TObjectPtr<UGraphIsland>>& GetIslands() const { return Islands; }
 	const FGraphProperties& GetProperties() const { return Properties; }
 
 	UGraphVertex* GetSafeVertexFromHandle(const FGraphVertexHandle& Handle) const;
-	UGraphEdge* GetSafeEdgeFromHandle(const FGraphEdgeHandle& Handle) const;
 	UGraphIsland* GetSafeIslandFromHandle(const FGraphIslandHandle& Handle) const;
 
 	GAMEPLAYGRAPH_API friend void operator<<(IGraphSerialization& Output, const UGraph& Graph);
@@ -130,22 +135,18 @@ public:
 
 protected:
 	/** When we add multiple edges into the graph. This function will ensure that the interactions we make externally are kept to a minimum. */
-	void MergeOrCreateIslands(const TArray<FGraphEdgeHandle>& InEdges);
+	void MergeOrCreateIslands(const TArray<FEdgeSpecifier>& InEdges);
 
 	/** After a change, this function will remove the island if it's empty or will attempt to split it into two smaller islands. */
 	void RemoveOrSplitIsland(TObjectPtr<UGraphIsland> Island);
 
 	/** Used for bulk loading in vertices/edges/islands. */
 	void ReserveVertices(int32 Delta);
-	void ReserveEdges(int32 Delta);
 	void ReserveIslands(int32 Delta);
 private:
 	UPROPERTY()
 	TMap<FGraphVertexHandle, TObjectPtr<UGraphVertex>> Vertices;
 
-	UPROPERTY()
-	TMap<FGraphEdgeHandle, TObjectPtr<UGraphEdge>> Edges;
-
 	UPROPERTY()
 	TMap<FGraphIslandHandle, TObjectPtr<UGraphIsland>> Islands;
 	FGraphProperties Properties;
@@ -154,22 +155,18 @@ private:
 	FGraphIslandHandle CreateIsland(TArray<FGraphVertexHandle> Nodes, FGraphUniqueIndex InUniqueIndex = FGraphUniqueIndex());
 
 	/** Creates an edge between the two nodes. */
-	FGraphEdgeHandle CreateEdge(FGraphVertexHandle Node1, FGraphVertexHandle Node2, FGraphUniqueIndex InUniqueIndex = FGraphUniqueIndex(), bool bMergeIslands = true);
+	bool CreateEdge(FGraphVertexHandle Node1, FGraphVertexHandle Node2, bool bMergeIslands = true);
 
 	/** Adds a node to the graph's node collection and modifies the NextAvailableNodeUniqueIndex as necessary to maintain its validity. */
 	void RegisterVertex(TObjectPtr<UGraphVertex> Node);
 
-	/** Add an edge to the graph and modifies the NextAvailableEdgeUniqueIndex as necessary to maintain its validity. */
-	void RegisterEdge(TObjectPtr<UGraphEdge> Edge);
-
 	/** Add an island to the graph and modifies the NextAvailableIslandUniqueIndex as necessary to maintain its validity. */
 	void RegisterIsland(TObjectPtr<UGraphIsland> Edge);
 
 	/** Need this so that users aren't able to pass in bHandleIslands = false. */
-	void RemoveEdgeInternal(const FGraphEdgeHandle& EdgeHandle, bool bHandleIslands);
+	void RemoveEdgeInternal(const FGraphVertexHandle& VertexHandleA, const FGraphVertexHandle& VertexHandleB, bool bHandleIslands);
 
 	/** Factory functions for creating an appropriately typed node/edge/island. */
 	virtual TObjectPtr<UGraphVertex> CreateTypedVertex() const;
-	virtual TObjectPtr<UGraphEdge> CreateTypedEdge() const;
 	virtual TObjectPtr<UGraphIsland> CreateTypedIsland() const;
 };
\ No newline at end of file
diff --git a/Engine/Plugins/Experimental/GameplayGraph/Source/GameplayGraph/Public/Graph/GraphDefaultSerialization.h b/Engine/Plugins/Experimental/GameplayGraph/Source/GameplayGraph/Public/Graph/GraphDefaultSerialization.h
--- a/Engine/Plugins/Experimental/GameplayGraph/Source/GameplayGraph/Public/Graph/GraphDefaultSerialization.h
+++ b/Engine/Plugins/Experimental/GameplayGraph/Source/GameplayGraph/Public/Graph/GraphDefaultSerialization.h
@@ -56,7 +56,7 @@ struct FSerializableGraph
 	TArray<FGraphVertexHandle> Vertices;
 
 	UPROPERTY(SaveGame)
-	TMap<FGraphEdgeHandle, FSerializedEdgeData> Edges;
+	TArray<FSerializedEdgeData> Edges;
 
 	UPROPERTY(SaveGame)
 	TMap<FGraphIslandHandle, FSerializedIslandData> Islands;
@@ -87,17 +87,12 @@ public:
 		Data.Vertices.Add(VertexHandle);
 	}
 
-	virtual void WriteGraphEdge(const FGraphEdgeHandle& EdgeHandle, const UGraphEdge* Edge) override
+	virtual void WriteGraphEdge(const FGraphVertexHandle& VertexHandleA, const FGraphVertexHandle& VertexHandleB) override
 	{
-		if (!ensure(Edge))
-		{
-			return;
-		}
-
 		FSerializedEdgeData Serialized;
-		Serialized.Node1 = Edge->NodeA();
-		Serialized.Node2 = Edge->NodeB();
-		Data.Edges.Emplace(EdgeHandle, MoveTemp(Serialized));
+		Serialized.Node1 = VertexHandleA;
+		Serialized.Node2 = VertexHandleB;
+		Data.Edges.Emplace(MoveTemp(Serialized));
 	}
 
 	virtual void WriteGraphIsland(const FGraphIslandHandle& IslandHandle, const UGraphIsland* Island) override
@@ -145,7 +140,7 @@ public:
 		for (const FGraphVertexHandle& SerializedHandle : Data.Vertices)
 		{
 			const FGraphVertexHandle FinalHandle = Lambda(SerializedHandle);
-			if (ensure(FinalHandle.IsValid()))
+			if (FinalHandle.IsValid())
 			{
 				OnDeserializedVertex(FinalHandle);
 			}
@@ -157,18 +152,14 @@ public:
 		return Data.Edges.Num();
 	}
 
-	virtual void ForEveryEdge(const TFunction<FGraphEdgeHandle(const FGraphEdgeHandle&, const FEdgeConstructionData&)>& Lambda) const override
+	virtual void ForEveryEdge(const TFunction<bool(const FEdgeSpecifier&)>& Lambda) const override
 	{
-		for (const TPair<FGraphEdgeHandle, FSerializedEdgeData>& Serialized : Data.Edges)
+		for (const FSerializedEdgeData& Serialized : Data.Edges)
 		{
-			FEdgeConstructionData Construction;
-			Construction.Vertex1 = Serialized.Value.Node1;
-			Construction.Vertex2 = Serialized.Value.Node2;
-
-			const FGraphEdgeHandle FinalHandle = Lambda(Serialized.Key, Construction);
-			if (ensure(FinalHandle.IsValid()))
+			const FEdgeSpecifier Construction{ Serialized.Node1, Serialized.Node2 };
+			if (Lambda(Construction))
 			{
-				OnDeserializedEdge(FinalHandle);
+				OnDeserializedEdge(Construction);
 			}
 		}
 	}
@@ -186,7 +177,7 @@ public:
 			Construction.Vertices = Serialized.Value.Vertices;
 
 			const FGraphIslandHandle FinalHandle = Lambda(Serialized.Key, Construction);
-			if (ensure(FinalHandle.IsValid()))
+			if (FinalHandle.IsValid())
 			{
 				OnDeserializedIsland(FinalHandle);
 			}
@@ -195,7 +186,7 @@ public:
 
 protected:
 	virtual void OnDeserializedVertex(const FGraphVertexHandle& VertexHandle) const {}
-	virtual void OnDeserializedEdge(const FGraphEdgeHandle& EdgeHandle) const {}
+	virtual void OnDeserializedEdge(const FEdgeSpecifier& Edge) const {}
 	virtual void OnDeserializedIsland(const FGraphIslandHandle& IslandHandle) const {}
 
 	const TSerializableGraph& Data;
diff --git a/Engine/Plugins/Experimental/GameplayGraph/Source/GameplayGraph/Public/Graph/GraphEdge.h b/Engine/Plugins/Experimental/GameplayGraph/Source/GameplayGraph/Public/Graph/GraphEdge.h
deleted file mode 100644
diff --git a/Engine/Plugins/Experimental/GameplayGraph/Source/GameplayGraph/Public/Graph/GraphHandle.h b/Engine/Plugins/Experimental/GameplayGraph/Source/GameplayGraph/Public/Graph/GraphHandle.h
--- a/Engine/Plugins/Experimental/GameplayGraph/Source/GameplayGraph/Public/Graph/GraphHandle.h
+++ b/Engine/Plugins/Experimental/GameplayGraph/Source/GameplayGraph/Public/Graph/GraphHandle.h
@@ -8,7 +8,6 @@
 
 class UGraph;
 class UGraphVertex;
-class UGraphEdge;
 class UGraphIsland;
 
 DECLARE_LOG_CATEGORY_EXTERN(LogGameplayGraph, Log, All);
@@ -147,22 +146,6 @@ struct GAMEPLAYGRAPH_API FGraphVertexHandle : public FGraphHandle
 	virtual bool HasElement() const override;
 };
 
-USTRUCT()
-struct GAMEPLAYGRAPH_API FGraphEdgeHandle : public FGraphHandle
-{
-	GENERATED_BODY()
-
-	static FGraphEdgeHandle Invalid;
-
-	FGraphEdgeHandle() = default;
-	FGraphEdgeHandle(FGraphUniqueIndex InUniqueIndex, UGraph* InGraph)
-		: FGraphHandle(InUniqueIndex, InGraph)
-	{}
-
-	UGraphEdge* GetEdge() const;
-	virtual bool HasElement() const override;
-};
-
 USTRUCT()
 struct GAMEPLAYGRAPH_API FGraphIslandHandle : public FGraphHandle
 {
diff --git a/Engine/Plugins/Experimental/GameplayGraph/Source/GameplayGraph/Public/Graph/GraphSerialization.h b/Engine/Plugins/Experimental/GameplayGraph/Source/GameplayGraph/Public/Graph/GraphSerialization.h
--- a/Engine/Plugins/Experimental/GameplayGraph/Source/GameplayGraph/Public/Graph/GraphSerialization.h
+++ b/Engine/Plugins/Experimental/GameplayGraph/Source/GameplayGraph/Public/Graph/GraphSerialization.h
@@ -4,12 +4,11 @@
 #include "Graph/GraphHandle.h"
 
 class UGraphVertex;
-class UGraphEdge;
 class UGraphIsland;
 
+struct FEdgeSpecifier;
 struct FGraphProperties;
 struct FGraphVertexHandle;
-struct FGraphEdgeHandle;
 struct FGraphIslandHandle;
 
 class IGraphSerialization
@@ -20,7 +19,7 @@ public:
 	virtual void Initialize(int32 NumVertices, int32 NumEdges, int32 NumIslands) = 0;
 	virtual void WriteGraphProperties(const FGraphProperties& Properties) = 0;
 	virtual void WriteGraphVertex(const FGraphVertexHandle& VertexHandle, const UGraphVertex* Vertex) = 0;
-	virtual void WriteGraphEdge(const FGraphEdgeHandle& EdgeHandle, const UGraphEdge* Edge) = 0;
+	virtual void WriteGraphEdge(const FGraphVertexHandle& VertexHandleA, const FGraphVertexHandle& VertexHandleB) = 0;
 	virtual void WriteGraphIsland(const FGraphIslandHandle& IslandHandle, const UGraphIsland* Island) = 0;
 };
 
@@ -33,14 +32,8 @@ public:
 	virtual int32 NumVertices() const = 0;
 	virtual void ForEveryVertex(const TFunction<FGraphVertexHandle(const FGraphVertexHandle&)>& Lambda) const = 0;
 
-	struct FEdgeConstructionData
-	{
-		FGraphVertexHandle Vertex1;
-		FGraphVertexHandle Vertex2;
-	};
-
 	virtual int32 NumEdges() const = 0;
-	virtual void ForEveryEdge(const TFunction<FGraphEdgeHandle(const FGraphEdgeHandle&, const FEdgeConstructionData&)>& Lambda) const = 0;
+	virtual void ForEveryEdge(const TFunction<bool(const FEdgeSpecifier&)>& Lambda) const = 0;
 
 	struct FIslandConstructionData
 	{
diff --git a/Engine/Plugins/Experimental/GameplayGraph/Source/GameplayGraph/Public/Graph/GraphVertex.h b/Engine/Plugins/Experimental/GameplayGraph/Source/GameplayGraph/Public/Graph/GraphVertex.h
--- a/Engine/Plugins/Experimental/GameplayGraph/Source/GameplayGraph/Public/Graph/GraphVertex.h
+++ b/Engine/Plugins/Experimental/GameplayGraph/Source/GameplayGraph/Public/Graph/GraphVertex.h
@@ -22,7 +22,6 @@ public:
 	}
 
 	bool HasEdgeTo(const FGraphVertexHandle& Other) const;
-	const FGraphEdgeHandle* FindEdgeTo(const FGraphVertexHandle& Other) const;
 	int32 NumEdges() const { return Edges.Num(); }
 
 	const FGraphIslandHandle& GetParentIsland() const { return ParentIsland; }
@@ -30,20 +29,13 @@ public:
 	template<typename TLambda>
 	void ForEachAdjacentVertex(TLambda&& Lambda)
 	{
-		for (const TPair<FGraphVertexHandle, FGraphEdgeHandle>& Kvp : Edges)
+		for (const FGraphVertexHandle& Vertex : Edges)
 		{
-			if constexpr (std::is_invocable_v<TLambda, const FGraphVertexHandle&, const FGraphEdgeHandle&>)
-			{
-				Lambda(Kvp.Key, Kvp.Value);
-			}
-			else
-			{
-				Lambda(Kvp.Key);
-			}
+			Lambda(Vertex);
 		}
 	}
 	
-	const TMap<FGraphVertexHandle, FGraphEdgeHandle>& GetEdges() const { return Edges; }
+	const TSet<FGraphVertexHandle>& GetEdges() const { return Edges; }
 
 	FOnGraphVertexRemoved OnVertexRemoved;
 	FOnGraphVertexParentIslandSet OnParentIslandSet;
@@ -52,15 +44,15 @@ public:
 	friend class UGraphIsland;
 protected:
 
-	void AddEdgeTo(const FGraphVertexHandle& Node, const FGraphEdgeHandle& Edge);
-	void RemoveEdge(const FGraphEdgeHandle& EdgeHandle);
+	void AddEdgeTo(const FGraphVertexHandle& Node);
+	void RemoveEdge(const FGraphVertexHandle& AdjacentVertexHandle);
 	void SetParentIsland(const FGraphIslandHandle& Island);
 
 	virtual void HandleOnVertexRemoved();
 private:
 
 	UPROPERTY()
-	TMap<FGraphVertexHandle, FGraphEdgeHandle> Edges;
+	TSet<FGraphVertexHandle> Edges;
 
 	UPROPERTY()
 	FGraphIslandHandle ParentIsland;
diff --git a/Engine/Plugins/Experimental/GameplayGraph/Tests/Private/GraphEdgeTests.cpp b/Engine/Plugins/Experimental/GameplayGraph/Tests/Private/GraphEdgeTests.cpp
deleted file mode 100644
diff --git a/Engine/Plugins/Experimental/GameplayGraph/Tests/Private/GraphElementTests.cpp b/Engine/Plugins/Experimental/GameplayGraph/Tests/Private/GraphElementTests.cpp
--- a/Engine/Plugins/Experimental/GameplayGraph/Tests/Private/GraphElementTests.cpp
+++ b/Engine/Plugins/Experimental/GameplayGraph/Tests/Private/GraphElementTests.cpp
@@ -3,7 +3,6 @@
 #include "TestHarness.h"
 
 #include "Graph/Graph.h"
-#include "Graph/GraphEdge.h"
 #include "Graph/GraphIsland.h"
 #include "Graph/GraphVertex.h"
 
@@ -15,12 +14,6 @@ TEST_CASE("Graph::Element::Element Type", "[graph][element]")
 		CHECK(Vertex->GetElementType() == EGraphElementType::Node);
 	}
 
-	SECTION("Edge")
-	{
-		UGraphEdge* Edge = NewObject<UGraphEdge>();
-		CHECK(Edge->GetElementType() == EGraphElementType::Edge);
-	}
-
 	SECTION("Island")
 	{
 		UGraphIsland* Island = NewObject<UGraphIsland>();
diff --git a/Engine/Plugins/Experimental/GameplayGraph/Tests/Private/GraphIslandTests.cpp b/Engine/Plugins/Experimental/GameplayGraph/Tests/Private/GraphIslandTests.cpp
--- a/Engine/Plugins/Experimental/GameplayGraph/Tests/Private/GraphIslandTests.cpp
+++ b/Engine/Plugins/Experimental/GameplayGraph/Tests/Private/GraphIslandTests.cpp
@@ -77,39 +77,29 @@ TEST_CASE_METHOD(FTestGraphBuilder, "Graph::Island::Bulk Create Edges", "[graph]
 	BuildLinearEdges(2);
 
 	{
-		TArray<FEdgeCreationParameters> AllEdgeParams;
+		TArray<FEdgeSpecifier> AllEdgeParams;
 		{
-			FEdgeCreationParameters Params;
-			Params.VertexHandle1 = VertexHandles[1];
-			Params.VertexHandle2 = VertexHandles[2];
+			FEdgeSpecifier Params{ VertexHandles[1], VertexHandles[2] };
 			AllEdgeParams.Add(Params);
 		}
 
 		{
-			FEdgeCreationParameters Params;
-			Params.VertexHandle1 = VertexHandles[1];
-			Params.VertexHandle2 = VertexHandles[3];
+			FEdgeSpecifier Params{ VertexHandles[1], VertexHandles[3] };
 			AllEdgeParams.Add(Params);
 		}
 
 		{
-			FEdgeCreationParameters Params;
-			Params.VertexHandle1 = VertexHandles[4];
-			Params.VertexHandle2 = VertexHandles[9];
+			FEdgeSpecifier Params{ VertexHandles[4], VertexHandles[9] };
 			AllEdgeParams.Add(Params);
 		}
 
 		{
-			FEdgeCreationParameters Params;
-			Params.VertexHandle1 = VertexHandles[5];
-			Params.VertexHandle2 = VertexHandles[6];
+			FEdgeSpecifier Params{ VertexHandles[5], VertexHandles[6] };
 			AllEdgeParams.Add(Params);
 		}
 
 		{
-			FEdgeCreationParameters Params;
-			Params.VertexHandle1 = VertexHandles[4];
-			Params.VertexHandle2 = VertexHandles[7];
+			FEdgeSpecifier Params{ VertexHandles[4], VertexHandles[7] };
 			AllEdgeParams.Add(Params);
 		}
 
@@ -272,93 +262,6 @@ TEST_CASE_METHOD(FTestGraphBuilder, "Graph::Island::Remove Vertex::Split", "[gra
 	IslandVertexParentIslandSanityCheck(IslandHandle2);
 }
 
-TEST_CASE_METHOD(FTestGraphBuilder, "Graph::Island::Remove Edge::Still Connected", "[graph][island]")
-{
-	PopulateVertices(4, true);
-	BuildFullyConnectedEdges(4);
-
-	REQUIRE(IslandHandles.Num() == 1);
-
-	UGraphIsland* Island1 = IslandHandles[0].GetIsland();
-	REQUIRE(Island1 != nullptr);
-	IslandVertexParentIslandSanityCheck(IslandHandles[0]);
-
-	{
-		CHECK(Island1->Num() == 4);
-		CHECK(Island1->GetVertices().Contains(VertexHandles[0]) == true);
-		CHECK(Island1->GetVertices().Contains(VertexHandles[1]) == true);
-		CHECK(Island1->GetVertices().Contains(VertexHandles[2]) == true);
-		CHECK(Island1->GetVertices().Contains(VertexHandles[3]) == true);
-	}
-
-	TArray<FGraphEdgeHandle> Edges = GetEdgesForVertex(VertexHandles[0]);
-	CHECK(Edges.Num() == 3);
-
-	Graph->RemoveEdge(Edges[0]);
-	IslandVertexParentIslandSanityCheck(IslandHandles[0]);
-	{
-		CHECK(Island1->Num() == 4);
-		CHECK(Island1->GetVertices().Contains(VertexHandles[0]) == true);
-		CHECK(Island1->GetVertices().Contains(VertexHandles[1]) == true);
-		CHECK(Island1->GetVertices().Contains(VertexHandles[2]) == true);
-		CHECK(Island1->GetVertices().Contains(VertexHandles[3]) == true);
-	}
-
-	Graph->RemoveEdge(Edges[1]);
-	IslandVertexParentIslandSanityCheck(IslandHandles[0]);
-	{
-		CHECK(Island1->Num() == 4);
-		CHECK(Island1->GetVertices().Contains(VertexHandles[0]) == true);
-		CHECK(Island1->GetVertices().Contains(VertexHandles[1]) == true);
-		CHECK(Island1->GetVertices().Contains(VertexHandles[2]) == true);
-		CHECK(Island1->GetVertices().Contains(VertexHandles[3]) == true);
-	}
-}
-
-TEST_CASE_METHOD(FTestGraphBuilder, "Graph::Island::Remove Edge::Split", "[graph][island]")
-{
-	PopulateVertices(5, true);
-	BuildLinearEdges(5);
-
-	REQUIRE(IslandHandles.Num() == 1);
-	CHECK(Graph->NumIslands() == 1);
-
-	TArray<FGraphEdgeHandle> Edges = GetEdgesForVertex(VertexHandles[2]);
-	CHECK(Edges.Num() == 2);
-	Graph->RemoveEdge(Edges[0]);
-	Graph->RemoveEdge(Edges[1]);
-
-	UGraphIsland* Island1 = VertexHandles[0].GetVertex()->GetParentIsland().GetIsland();
-	REQUIRE(Island1 != nullptr);
-	{
-		CHECK(Island1->Num() == 2);
-		CHECK(Island1->GetVertices().Contains(VertexHandles[0]) == true);
-		CHECK(Island1->GetVertices().Contains(VertexHandles[1]) == true);
-	}
-
-	UGraphIsland* Island2 = VertexHandles[2].GetVertex()->GetParentIsland().GetIsland();
-	REQUIRE(Island2 != nullptr);
-	CHECK(Island1->Handle() != Island2->Handle());
-	{
-		CHECK(Island2->Num() == 1);
-		CHECK(Island2->GetVertices().Contains(VertexHandles[2]) == true);
-	}
-
-	UGraphIsland* Island3 = VertexHandles[3].GetVertex()->GetParentIsland().GetIsland();
-	REQUIRE(Island3 != nullptr);
-	CHECK(Island1->Handle() != Island3->Handle());
-	CHECK(Island2->Handle() != Island3->Handle());
-	{
-		CHECK(Island3->Num() == 2);
-		CHECK(Island3->GetVertices().Contains(VertexHandles[3]) == true);
-		CHECK(Island3->GetVertices().Contains(VertexHandles[4]) == true);
-	}
-
-	IslandVertexParentIslandSanityCheck(Island1->Handle());
-	IslandVertexParentIslandSanityCheck(Island2->Handle());
-	IslandVertexParentIslandSanityCheck(Island3->Handle());
-}
-
 TEST_CASE_METHOD(FTestGraphBuilder, "Graph::Island::Operation Allowed::Flag", "[graph][island]")
 {
 	PopulateVertices(2, true);
@@ -424,22 +327,18 @@ TEST_CASE_METHOD(FTestGraphBuilder, "Graph::Island::Operation Allowed::Disable A
 
 	FGraphVertexHandle NewVertex = Graph->CreateVertex();
 
-	FEdgeCreationParameters Params;
-	Params.VertexHandle1 = NewVertex;
-	Params.VertexHandle2 = VertexHandles[0];
-
-	TArray<FGraphEdgeHandle> Edges;
-	Graph->CreateBulkEdges({ Params }, &Edges);
+	FEdgeSpecifier Params{ NewVertex, VertexHandles[0] };
+	Graph->CreateBulkEdges({ Params });
 
-	CHECK(Edges.Num() == 0);
 	CHECK(Island->Num() == 4);
 	CHECK(Island->GetVertices().Contains(NewVertex) == false);
 	CHECK(NewVertex.GetVertex()->GetParentIsland() != Island->Handle());
 	CHECK(NewVertex.GetVertex()->GetParentIsland() == FGraphIslandHandle::Invalid);
 	CHECK(Graph->NumVertices() == 5);
-	CHECK(Graph->NumEdges() == 3);
 	CHECK(Graph->NumIslands() == 1);
 	IslandVertexParentIslandSanityCheck(IslandHandles[0]);
+	CHECK(NewVertex.GetVertex()->HasEdgeTo(VertexHandles[0]) == false);
+	CHECK(VertexHandles[0].GetVertex()->HasEdgeTo(NewVertex) == false);
 }
 
 TEST_CASE_METHOD(FTestGraphBuilder, "Graph::Island::Operation Allowed::Disable Split", "[graph][island]")
@@ -456,7 +355,6 @@ TEST_CASE_METHOD(FTestGraphBuilder, "Graph::Island::Operation Allowed::Disable S
 
 	CHECK(Island->Num() == 6);
 	CHECK(Graph->NumVertices() == 6);
-	CHECK(Graph->NumEdges() == 5);
 	CHECK(Graph->NumIslands() == 1);
 
 	UGraphVertex* Vertex = VertexHandles[2].GetVertex();
@@ -465,42 +363,6 @@ TEST_CASE_METHOD(FTestGraphBuilder, "Graph::Island::Operation Allowed::Disable S
 	UGraphIsland* Island1 = nullptr;
 	UGraphIsland* Island2 = nullptr;
 
-	SECTION("Remove Edge")
-	{
-		const FGraphEdgeHandle* EdgeHandle = Vertex->FindEdgeTo(VertexHandles[3]);
-		REQUIRE(EdgeHandle != nullptr);
-		Graph->RemoveEdge(*EdgeHandle);
-
-		IslandVertexParentIslandSanityCheck(IslandHandles[0]);
-		CHECK(Island->GetVertices().Contains(VertexHandles[0]) == true);
-		CHECK(Island->GetVertices().Contains(VertexHandles[1]) == true);
-		CHECK(Island->GetVertices().Contains(VertexHandles[2]) == true);
-		CHECK(Island->GetVertices().Contains(VertexHandles[3]) == true);
-		CHECK(Island->GetVertices().Contains(VertexHandles[4]) == true);
-		CHECK(Island->GetVertices().Contains(VertexHandles[5]) == true);
-		CHECK(Island->Num() == 6);
-		CHECK(Graph->NumVertices() == 6);
-		CHECK(Graph->NumEdges() == 4);
-
-		CHECK(Graph->NumIslands() == 1);
-		Island->SetOperationAllowed(EGraphIslandOperations::Split, true);
-		Graph->RefreshIslandConnectivity(Island->Handle());
-
-		Island1 = VertexHandles[0].GetVertex()->GetParentIsland().GetIsland();
-		Island2 = VertexHandles[5].GetVertex()->GetParentIsland().GetIsland();
-
-		CHECK(Island1->Num() == 3);
-		CHECK(Island1->GetVertices().Contains(VertexHandles[0]) == true);
-		CHECK(Island1->GetVertices().Contains(VertexHandles[1]) == true);
-		CHECK(Island1->GetVertices().Contains(VertexHandles[2]) == true);
-		CHECK(Island2->Num() == 3);
-		CHECK(Island2->GetVertices().Contains(VertexHandles[3]) == true);
-		CHECK(Island2->GetVertices().Contains(VertexHandles[4]) == true);
-		CHECK(Island2->GetVertices().Contains(VertexHandles[5]) == true);
-		CHECK(Graph->NumVertices() == 6);
-		CHECK(Graph->NumEdges() == 4);
-	}
-
 	SECTION("Remove Vertex")
 	{
 		Graph->RemoveVertex(VertexHandles[2]);
@@ -512,7 +374,6 @@ TEST_CASE_METHOD(FTestGraphBuilder, "Graph::Island::Operation Allowed::Disable S
 		CHECK(Island->GetVertices().Contains(VertexHandles[5]) == true);
 		CHECK(Island->Num() == 5);
 		CHECK(Graph->NumVertices() == 5);
-		CHECK(Graph->NumEdges() == 3);
 
 		CHECK(Graph->NumIslands() == 1);
 		Island->SetOperationAllowed(EGraphIslandOperations::Split, true);
@@ -529,7 +390,6 @@ TEST_CASE_METHOD(FTestGraphBuilder, "Graph::Island::Operation Allowed::Disable S
 		CHECK(Island2->GetVertices().Contains(VertexHandles[4]) == true);
 		CHECK(Island2->GetVertices().Contains(VertexHandles[5]) == true);
 		CHECK(Graph->NumVertices() == 5);
-		CHECK(Graph->NumEdges() == 3);
 	}
 
 	CHECK(Graph->NumIslands() == 2);
@@ -548,17 +408,11 @@ TEST_CASE_METHOD(FTestGraphBuilder, "Graph::Island::Operation Allowed::Disable M
 
 	auto RunTest = [this]()
 	{
-		FEdgeCreationParameters Params;
-		Params.VertexHandle1 = VertexHandles[0];
-		Params.VertexHandle2 = VertexHandles[2];
-
-		TArray<FGraphEdgeHandle> Edges;
-		Graph->CreateBulkEdges({ Params }, &Edges);
-
-		CHECK(Edges.Num() == 0);
+		FEdgeSpecifier Params{ VertexHandles[0], VertexHandles[2] };
+		Graph->CreateBulkEdges({ Params });
 		CHECK(Graph->NumVertices() == 4);
-		CHECK(Graph->NumEdges() == 2);
 		CHECK(Graph->NumIslands() == 2);
+		VerifyEdges({ Params }, false);
 
 		UGraphIsland* Island0 = VertexHandles[0].GetVertex()->GetParentIsland().GetIsland();
 		REQUIRE(Island0 != nullptr);
@@ -573,7 +427,6 @@ TEST_CASE_METHOD(FTestGraphBuilder, "Graph::Island::Operation Allowed::Disable M
 		CHECK(Island1->GetVertices().Contains(VertexHandles[2]) == true);
 		CHECK(Island1->GetVertices().Contains(VertexHandles[3]) == true);
 
-
 		IslandVertexParentIslandSanityCheck(Island0->Handle());
 		IslandVertexParentIslandSanityCheck(Island1->Handle());
 	};
@@ -622,21 +475,15 @@ TEST_CASE_METHOD(FTestGraphBuilder, "Graph::Island::Operation Allowed::Disable M
 	auto RunTest = [this]()
 	{
 		FGraphVertexHandle NewVertex = Graph->CreateVertex();
-		FEdgeCreationParameters Params1;
-		Params1.VertexHandle1 = VertexHandles[0];
-		Params1.VertexHandle2 = NewVertex;
-
-		FEdgeCreationParameters Params2;
-		Params2.VertexHandle1 = VertexHandles[3];
-		Params2.VertexHandle2 = NewVertex;
-
-		TArray<FGraphEdgeHandle> Edges;
-		Graph->CreateBulkEdges({ Params1, Params2 }, &Edges);
+		FEdgeSpecifier Params1{VertexHandles[0], NewVertex};
+		FEdgeSpecifier Params2{VertexHandles[3], NewVertex};
+		Graph->CreateBulkEdges({ Params1, Params2 });
 		Graph->FinalizeVertex(NewVertex);
 
-		CHECK(Edges.Num() == 1);
+		VerifyEdges( { Params1 }, true);
+		VerifyEdges( { Params2 }, false);
+
 		CHECK(Graph->NumVertices() == 5);
-		CHECK(Graph->NumEdges() == 3);
 		CHECK(Graph->NumIslands() == 2);
 
 		UGraphIsland* Island0 = VertexHandles[0].GetVertex()->GetParentIsland().GetIsland();
@@ -739,9 +586,7 @@ TEST_CASE_METHOD(FTestGraphBuilder, "Graph::Island::Events::Vertex Added", "[gra
 	);
 
 	FGraphVertexHandle NewVertex = Graph->CreateVertex();
-	FEdgeCreationParameters Params;
-	Params.VertexHandle1 = NewVertex;
-	Params.VertexHandle2 = VertexHandles[0];
+	FEdgeSpecifier Params{ NewVertex, VertexHandles[0] };
 	Graph->CreateBulkEdges({ Params });
 
 	CHECK(CallbackIslandHandle == Island->Handle());
@@ -775,15 +620,6 @@ TEST_CASE_METHOD(FTestGraphBuilder, "Graph::Island::Events::Vertex Removed", "[g
 		CHECK(CallbackIslandHandle == Island->Handle());
 		CHECK(CallbackVertexHandle == VertexHandles[0]);
 	}
-
-	SECTION("Remove Edge")
-	{
-		TArray<FGraphEdgeHandle> Edges = GetEdgesForVertex(VertexHandles[0]);
-		CHECK(Edges.Num() == 1);
-		Graph->RemoveEdge(Edges[0]);
-		CHECK(CallbackIslandHandle == Island->Handle());
-		CHECK(CallbackVertexHandle == VertexHandles[0]);
-	}
 }
 
 TEST_CASE_METHOD(FTestGraphBuilder, "Graph::Island::Events::On Destroyed", "[graph][island]")
diff --git a/Engine/Plugins/Experimental/GameplayGraph/Tests/Private/GraphSerializationTests.cpp b/Engine/Plugins/Experimental/GameplayGraph/Tests/Private/GraphSerializationTests.cpp
--- a/Engine/Plugins/Experimental/GameplayGraph/Tests/Private/GraphSerializationTests.cpp
+++ b/Engine/Plugins/Experimental/GameplayGraph/Tests/Private/GraphSerializationTests.cpp
@@ -10,7 +10,6 @@ TEST_CASE_METHOD(FTestGraphBuilder, "Graph::Serialization::Write", "[graph]")
 	BuildLinearEdges(3);
 
 	CHECK(Graph->NumVertices() == 6);
-	CHECK(Graph->NumEdges() == 4);
 	CHECK(Graph->NumIslands() == 2);
 
 	FDefaultGraphSerialization Serializer;
@@ -26,23 +25,34 @@ TEST_CASE_METHOD(FTestGraphBuilder, "Graph::Serialization::Write", "[graph]")
 	CHECK(SerializedGraph.Vertices.Contains(VertexHandles[5]) == true);
 
 	CHECK(SerializedGraph.Edges.Num() == 4);
+
+	TSet<FEdgeSpecifier> SerializedEdges;
+	TSet<FEdgeSpecifier> TestEdges;
+
+	for (const FSerializedEdgeData& EdgeData : SerializedGraph.Edges)
+	{
+		const FEdgeSpecifier SerializedEdge{EdgeData.Node1, EdgeData.Node2};
+		SerializedEdges.Add(SerializedEdge);
+	}
+
 	for (const FGraphVertexHandle& VertexHandle : VertexHandles)
 	{
 		UGraphVertex* Vertex = VertexHandle.GetVertex();
 		REQUIRE(Vertex != nullptr);
 
 		Vertex->ForEachAdjacentVertex(
-			[&SerializedGraph](const FGraphVertexHandle& NeighborVertexHandle, const FGraphEdgeHandle& EdgeHandle)
+			[&TestEdges, &VertexHandle](const FGraphVertexHandle& NeighborVertexHandle)
 			{
-				REQUIRE(SerializedGraph.Edges.Contains(EdgeHandle) == true);
-				const FSerializedEdgeData& EdgeData = SerializedGraph.Edges[EdgeHandle];
-				UGraphEdge* Edge = EdgeHandle.GetEdge();
-				REQUIRE(Edge != nullptr);
-				CHECK(EdgeData.Node1 == Edge->NodeA());
-				CHECK(EdgeData.Node2 == Edge->NodeB());
+				const FEdgeSpecifier TestEdge{VertexHandle, NeighborVertexHandle};
+				TestEdges.Add(TestEdge);
 			}
 		);
 	}
+	CHECK(SerializedEdges.Num() == TestEdges.Num());
+	for (const FEdgeSpecifier& Edge : TestEdges)
+	{
+		CHECK(SerializedEdges.Contains(Edge) == true);
+	}
 
 	CHECK(SerializedGraph.Islands.Num() == 2);
 	CHECK(SerializedGraph.Islands.Contains(IslandHandles[0]) == true);
@@ -75,10 +85,10 @@ TEST_CASE_METHOD(FTestGraphBuilder, "Graph::Serialization::Read (normal)", "[gra
 	SerializedGraph.Vertices.Add(FGraphVertexHandle { FGraphUniqueIndex::CreateUniqueIndex(), nullptr });
 	SerializedGraph.Vertices.Add(FGraphVertexHandle { FGraphUniqueIndex::CreateUniqueIndex(), nullptr });
 
-	SerializedGraph.Edges.Add(FGraphEdgeHandle { FGraphUniqueIndex::CreateUniqueIndex(), nullptr }, FSerializedEdgeData{ SerializedGraph.Vertices[0], SerializedGraph.Vertices[1] });
-	SerializedGraph.Edges.Add(FGraphEdgeHandle { FGraphUniqueIndex::CreateUniqueIndex(), nullptr }, FSerializedEdgeData{ SerializedGraph.Vertices[1], SerializedGraph.Vertices[2] });
-	SerializedGraph.Edges.Add(FGraphEdgeHandle { FGraphUniqueIndex::CreateUniqueIndex(), nullptr }, FSerializedEdgeData{ SerializedGraph.Vertices[3], SerializedGraph.Vertices[4] });
-	SerializedGraph.Edges.Add(FGraphEdgeHandle { FGraphUniqueIndex::CreateUniqueIndex(), nullptr }, FSerializedEdgeData{ SerializedGraph.Vertices[4], SerializedGraph.Vertices[5] });
+	SerializedGraph.Edges.Add(FSerializedEdgeData{ SerializedGraph.Vertices[0], SerializedGraph.Vertices[1] });
+	SerializedGraph.Edges.Add(FSerializedEdgeData{ SerializedGraph.Vertices[1], SerializedGraph.Vertices[2] });
+	SerializedGraph.Edges.Add(FSerializedEdgeData{ SerializedGraph.Vertices[3], SerializedGraph.Vertices[4] });
+	SerializedGraph.Edges.Add(FSerializedEdgeData{ SerializedGraph.Vertices[4], SerializedGraph.Vertices[5] });
 
 	SerializedGraph.Islands.Add(FGraphIslandHandle { FGraphUniqueIndex::CreateUniqueIndex(), nullptr }, FSerializedIslandData { TArray{SerializedGraph.Vertices[0], SerializedGraph.Vertices[1], SerializedGraph.Vertices[2]} });
 	SerializedGraph.Islands.Add(FGraphIslandHandle { FGraphUniqueIndex::CreateUniqueIndex(), nullptr }, FSerializedIslandData { TArray{SerializedGraph.Vertices[3], SerializedGraph.Vertices[4], SerializedGraph.Vertices[5]} });
@@ -88,7 +98,6 @@ TEST_CASE_METHOD(FTestGraphBuilder, "Graph::Serialization::Read (normal)", "[gra
 
 	CHECK(SerializedGraph.Properties == Graph->GetProperties());
 	CHECK(Graph->NumVertices() == 6);
-	CHECK(Graph->NumEdges() == 4);
 	CHECK(Graph->NumIslands() == 2);
 
 	int32 Index = 0;
@@ -101,19 +110,19 @@ TEST_CASE_METHOD(FTestGraphBuilder, "Graph::Serialization::Read (normal)", "[gra
 		CHECK(SerializedVertexHandle == LoadedVertex->Handle());
 	}
 
-	for (const TPair<FGraphEdgeHandle, FSerializedEdgeData>& Edge: SerializedGraph.Edges)
+	for (const FSerializedEdgeData& Edge: SerializedGraph.Edges)
 	{
-		CHECK(Graph->GetEdges().Contains(Edge.Key) == true);
+		CHECK(Graph->GetVertices().Contains(Edge.Node1) == true);
+		CHECK(Graph->GetVertices().Contains(Edge.Node2) == true);
+
+		UGraphVertex* Vertex1 = Graph->GetVertices().FindRef(Edge.Node1);
+		REQUIRE(Vertex1 != nullptr);
 
-		CHECK(Graph->GetVertices().Contains(Edge.Value.Node1) == true);
-		CHECK(Graph->GetVertices().Contains(Edge.Value.Node2) == true);
+		UGraphVertex* Vertex2 = Graph->GetVertices().FindRef(Edge.Node2);
+		REQUIRE(Vertex2 != nullptr);
 
-		UGraphEdge* LoadedEdge = Graph->GetEdges().FindRef(Edge.Key);
-		REQUIRE(LoadedEdge != nullptr);
-		CHECK(LoadedEdge->ContainsNode(Edge.Value.Node1) == true);
-		CHECK(LoadedEdge->NodeA().IsComplete() == true);
-		CHECK(LoadedEdge->ContainsNode(Edge.Value.Node2) == true);
-		CHECK(LoadedEdge->NodeB().IsComplete() == true);
+		CHECK(Vertex1->HasEdgeTo(Edge.Node2) == true);
+		CHECK(Vertex2->HasEdgeTo(Edge.Node1) == true);
 	}
 
 	for (const TPair<FGraphIslandHandle, FSerializedIslandData>& Island : SerializedGraph.Islands)
@@ -145,23 +154,6 @@ TEST_CASE_METHOD(FTestGraphBuilder, "Graph::Serialization::Read::Errors::Vertex:
 	Deserializer >> *Graph;
 	CHECK(SerializedGraph.Properties == Graph->GetProperties());
 	CHECK(Graph->NumVertices() == 0);
-	CHECK(Graph->NumEdges() == 0);
-	CHECK(Graph->NumIslands() == 0);
-}
-
-TEST_CASE_METHOD(FTestGraphBuilder, "Graph::Serialization::Read::Errors::Edge::Invalid Index", "[graph]")
-{
-	FSerializableGraph SerializedGraph;
-	SerializedGraph.Properties.bGenerateIslands = true;
-	SerializedGraph.Vertices.Add(FGraphVertexHandle { FGraphUniqueIndex::CreateUniqueIndex(), nullptr });
-	SerializedGraph.Vertices.Add(FGraphVertexHandle { FGraphUniqueIndex::CreateUniqueIndex(), nullptr });
-	SerializedGraph.Edges.Add(FGraphEdgeHandle { FGraphUniqueIndex{FGuid{0, 0, 0, 0}}, nullptr }, FSerializedEdgeData{ SerializedGraph.Vertices[0], SerializedGraph.Vertices[1] });
-
-	FDefaultGraphDeserialization Deserializer { SerializedGraph };
-	Deserializer >> *Graph;
-	CHECK(SerializedGraph.Properties == Graph->GetProperties());
-	CHECK(Graph->NumVertices() == 2);
-	CHECK(Graph->NumEdges() == 0);
 	CHECK(Graph->NumIslands() == 0);
 }
 
@@ -171,13 +163,12 @@ TEST_CASE_METHOD(FTestGraphBuilder, "Graph::Serialization::Read::Errors::Edge::I
 	SerializedGraph.Properties.bGenerateIslands = true;
 	SerializedGraph.Vertices.Add(FGraphVertexHandle { FGraphUniqueIndex{FGuid{0, 0, 0, 1}}, nullptr });
 	SerializedGraph.Vertices.Add(FGraphVertexHandle { FGraphUniqueIndex{FGuid{0, 0, 0, 2}}, nullptr });
-	SerializedGraph.Edges.Add(FGraphEdgeHandle { FGraphUniqueIndex::CreateUniqueIndex(), nullptr }, FSerializedEdgeData{ SerializedGraph.Vertices[0], FGraphVertexHandle { FGraphUniqueIndex{FGuid{0, 0, 0, 3}}, nullptr } });
+	SerializedGraph.Edges.Add(FSerializedEdgeData{ SerializedGraph.Vertices[0], FGraphVertexHandle { FGraphUniqueIndex{FGuid{0, 0, 0, 3}}, nullptr } });
 
 	FDefaultGraphDeserialization Deserializer { SerializedGraph };
 	Deserializer >> *Graph;
 	CHECK(SerializedGraph.Properties == Graph->GetProperties());
 	CHECK(Graph->NumVertices() == 2);
-	CHECK(Graph->NumEdges() == 0);
 	CHECK(Graph->NumIslands() == 0);
 }
 
@@ -187,14 +178,13 @@ TEST_CASE_METHOD(FTestGraphBuilder, "Graph::Serialization::Read::Errors::Island:
 	SerializedGraph.Properties.bGenerateIslands = true;
 	SerializedGraph.Vertices.Add(FGraphVertexHandle { FGraphUniqueIndex::CreateUniqueIndex(), nullptr });
 	SerializedGraph.Vertices.Add(FGraphVertexHandle { FGraphUniqueIndex::CreateUniqueIndex(), nullptr });
-	SerializedGraph.Edges.Add(FGraphEdgeHandle { FGraphUniqueIndex::CreateUniqueIndex(), nullptr }, FSerializedEdgeData{ SerializedGraph.Vertices[0], SerializedGraph.Vertices[1] });
+	SerializedGraph.Edges.Add(FSerializedEdgeData{ SerializedGraph.Vertices[0], SerializedGraph.Vertices[1] });
 	SerializedGraph.Islands.Add(FGraphIslandHandle { FGraphUniqueIndex{FGuid{0, 0, 0, 0}}, nullptr }, FSerializedIslandData { TArray{SerializedGraph.Vertices[0], SerializedGraph.Vertices[1] } });
 
 	FDefaultGraphDeserialization Deserializer { SerializedGraph };
 	Deserializer >> *Graph;
 	CHECK(SerializedGraph.Properties == Graph->GetProperties());
 	CHECK(Graph->NumVertices() == 2);
-	CHECK(Graph->NumEdges() == 1);
 	CHECK(Graph->NumIslands() == 1);
 	CHECK(Graph->GetIslands().Contains(FGraphIslandHandle { FGraphUniqueIndex { FGuid { 0, 0, 0, 0 } }, nullptr }) == false);
 }
@@ -205,14 +195,13 @@ TEST_CASE_METHOD(FTestGraphBuilder, "Graph::Serialization::Read::Errors::Island:
 	SerializedGraph.Properties.bGenerateIslands = true;
 	SerializedGraph.Vertices.Add(FGraphVertexHandle { FGraphUniqueIndex{FGuid{0, 0, 0, 1}}, nullptr });
 	SerializedGraph.Vertices.Add(FGraphVertexHandle { FGraphUniqueIndex{FGuid{0, 0, 0, 2}}, nullptr });
-	SerializedGraph.Edges.Add(FGraphEdgeHandle { FGraphUniqueIndex::CreateUniqueIndex(), nullptr }, FSerializedEdgeData{ SerializedGraph.Vertices[0], SerializedGraph.Vertices[1] });
+	SerializedGraph.Edges.Add(FSerializedEdgeData{ SerializedGraph.Vertices[0], SerializedGraph.Vertices[1] });
 	SerializedGraph.Islands.Add(FGraphIslandHandle { FGraphUniqueIndex{FGuid{1, 0, 0, 0}}, nullptr }, FSerializedIslandData { TArray{ SerializedGraph.Vertices[0], SerializedGraph.Vertices[1], FGraphVertexHandle { FGraphUniqueIndex{FGuid{0, 0, 0, 3}}, nullptr } } });
 
 	FDefaultGraphDeserialization Deserializer { SerializedGraph };
 	Deserializer >> *Graph;
 	CHECK(SerializedGraph.Properties == Graph->GetProperties());
 	CHECK(Graph->NumVertices() == 2);
-	CHECK(Graph->NumEdges() == 1);
 	CHECK(Graph->NumIslands() == 1);
 
 	UGraphIsland* LoadedIsland = Graph->GetIslands().FindRef(FGraphIslandHandle {FGraphUniqueIndex{FGuid{1, 0, 0, 0}}, nullptr });
@@ -231,8 +220,8 @@ TEST_CASE_METHOD(FTestGraphBuilder, "Graph::Serialization::Read::Errors::Island:
 	SerializedGraph.Vertices.Add(FGraphVertexHandle { FGraphUniqueIndex::CreateUniqueIndex(), nullptr });
 	SerializedGraph.Vertices.Add(FGraphVertexHandle { FGraphUniqueIndex::CreateUniqueIndex(), nullptr });
 
-	SerializedGraph.Edges.Add(FGraphEdgeHandle { FGraphUniqueIndex::CreateUniqueIndex(), nullptr }, FSerializedEdgeData{ SerializedGraph.Vertices[0], SerializedGraph.Vertices[1] });
-	SerializedGraph.Edges.Add(FGraphEdgeHandle { FGraphUniqueIndex::CreateUniqueIndex(), nullptr }, FSerializedEdgeData{ SerializedGraph.Vertices[1], SerializedGraph.Vertices[2] });
+	SerializedGraph.Edges.Add(FSerializedEdgeData{ SerializedGraph.Vertices[0], SerializedGraph.Vertices[1] });
+	SerializedGraph.Edges.Add(FSerializedEdgeData{ SerializedGraph.Vertices[1], SerializedGraph.Vertices[2] });
 
 	SerializedGraph.Islands.Add(FGraphIslandHandle { FGraphUniqueIndex{FGuid{1, 0, 0, 0}}, nullptr }, FSerializedIslandData { TArray { SerializedGraph.Vertices[0], SerializedGraph.Vertices[1] } });
 	SerializedGraph.Islands.Add(FGraphIslandHandle { FGraphUniqueIndex{FGuid{2, 0, 0, 0}}, nullptr }, FSerializedIslandData { TArray { SerializedGraph.Vertices[2] } });
@@ -241,7 +230,6 @@ TEST_CASE_METHOD(FTestGraphBuilder, "Graph::Serialization::Read::Errors::Island:
 	Deserializer >> *Graph;
 	CHECK(SerializedGraph.Properties == Graph->GetProperties());
 	CHECK(Graph->NumVertices() == 3);
-	CHECK(Graph->NumEdges() == 2);
 	CHECK(Graph->NumIslands() == 1);
 
 	CHECK(Graph->GetIslands().Contains(FGraphIslandHandle { FGraphUniqueIndex { FGuid { 1, 0, 0, 0 } }, nullptr }) == true);
@@ -263,7 +251,7 @@ TEST_CASE_METHOD(FTestGraphBuilder, "Graph::Serialization::Read::Errors::Island:
 	SerializedGraph.Vertices.Add(FGraphVertexHandle { FGraphUniqueIndex::CreateUniqueIndex(), nullptr });
 	SerializedGraph.Vertices.Add(FGraphVertexHandle { FGraphUniqueIndex::CreateUniqueIndex(), nullptr });
 
-	SerializedGraph.Edges.Add(FGraphEdgeHandle { FGraphUniqueIndex::CreateUniqueIndex(), nullptr }, FSerializedEdgeData{ SerializedGraph.Vertices[0], SerializedGraph.Vertices[1] });
+	SerializedGraph.Edges.Add(FSerializedEdgeData{ SerializedGraph.Vertices[0], SerializedGraph.Vertices[1] });
 
 	SerializedGraph.Islands.Add(FGraphIslandHandle { FGraphUniqueIndex{FGuid{1, 0, 0, 0}}, nullptr }, FSerializedIslandData { TArray { SerializedGraph.Vertices[0], SerializedGraph.Vertices[1], SerializedGraph.Vertices[2] } });
 
@@ -271,7 +259,6 @@ TEST_CASE_METHOD(FTestGraphBuilder, "Graph::Serialization::Read::Errors::Island:
 	Deserializer >> *Graph;
 	CHECK(SerializedGraph.Properties == Graph->GetProperties());
 	CHECK(Graph->NumVertices() == 3);
-	CHECK(Graph->NumEdges() == 1);
 	CHECK(Graph->NumIslands() == 2);
 
 	CHECK(Graph->GetIslands().Contains(FGraphIslandHandle { FGraphUniqueIndex { FGuid { 1, 0, 0, 0 } }, nullptr }) == true);
diff --git a/Engine/Plugins/Experimental/GameplayGraph/Tests/Private/GraphTests.cpp b/Engine/Plugins/Experimental/GameplayGraph/Tests/Private/GraphTests.cpp
--- a/Engine/Plugins/Experimental/GameplayGraph/Tests/Private/GraphTests.cpp
+++ b/Engine/Plugins/Experimental/GameplayGraph/Tests/Private/GraphTests.cpp
@@ -34,7 +34,6 @@ TEST_CASE("Graph::Default Constructor", "[graph]")
 	SECTION("State")
 	{
 		CHECK(Graph->NumVertices() == 0);
-		CHECK(Graph->NumEdges() == 0);
 		CHECK(Graph->NumIslands() == 0);
 	}
 }
@@ -67,7 +66,6 @@ TEST_CASE_METHOD(FTestGraphBuilder, "Graph::Create Vertices", "[graph]")
 	{
 		CHECK(Node1.GetUniqueIndex() == Node1Index);
 		CHECK(Graph->NumVertices() == 1);
-		CHECK(Graph->NumEdges() == 0);
 		CHECK(Graph->NumIslands() == 0);
 		CHECK(Graph->GetVertices().Contains(Node1) == true);
 	}
@@ -77,7 +75,6 @@ TEST_CASE_METHOD(FTestGraphBuilder, "Graph::Create Vertices", "[graph]")
 	{
 		CHECK(Node2.GetUniqueIndex() == Node2Index);
 		CHECK(Graph->NumVertices() == 2);
-		CHECK(Graph->NumEdges() == 0);
 		CHECK(Graph->NumIslands() == 0);
 		CHECK(Graph->GetVertices().Contains(Node2) == true);
 	}
@@ -125,7 +122,6 @@ TEST_CASE_METHOD(FTestGraphBuilder, "Graph::Create Vertices", "[graph]")
 	{
 		Graph->FinalizeVertex(Node1);
 		CHECK(Graph->NumVertices() == 2);
-		CHECK(Graph->NumEdges() == 0);
 		CHECK(Graph->NumIslands() == 1);
 	}
 }
@@ -137,7 +133,6 @@ TEST_CASE_METHOD(FTestGraphBuilder, "Graph::Create Vertices::Duplicate", "[graph
 	{
 		CHECK(Node1.GetUniqueIndex() == Node1Index);
 		CHECK(Graph->NumVertices() == 1);
-		CHECK(Graph->NumEdges() == 0);
 		CHECK(Graph->NumIslands() == 0);
 		CHECK(Graph->GetVertices().Contains(Node1) == true);
 	}
@@ -147,7 +142,6 @@ TEST_CASE_METHOD(FTestGraphBuilder, "Graph::Create Vertices::Duplicate", "[graph
 		CHECK(Node2.IsValid() == false);
 		CHECK(Node2.IsComplete() == false);
 		CHECK(Graph->NumVertices() == 1);
-		CHECK(Graph->NumEdges() == 0);
 		CHECK(Graph->NumIslands() == 0);
 		CHECK(Graph->GetVertices().Contains(Node2) == false);
 	}
@@ -158,87 +152,63 @@ TEST_CASE_METHOD(FTestGraphBuilder, "Graph::Create Edges::Single", "[graph]")
 	PopulateVertices(4, false);
 
 	{
-		FEdgeCreationParameters Params;
-		Params.VertexHandle1 = VertexHandles[1];
-		Params.VertexHandle2 = VertexHandles[0];
-		Params.EdgeIndex = FGraphUniqueIndex{ FGuid{ 1, 1, 1, 1 } };
-
-		TArray<FGraphEdgeHandle> Edges;
-		Graph->CreateBulkEdges({ Params }, &Edges);
-		CHECK(Graph->NumEdges() == 1);
+		FEdgeSpecifier Params{ VertexHandles[1], VertexHandles[0] };
+		Graph->CreateBulkEdges({ Params });
 		CHECK(Graph->NumIslands() == 1);
-
-		REQUIRE(Edges.Num() == 1);
-		CHECK(Edges[0].GetUniqueIndex() == Params.EdgeIndex);
-		CHECK(Edges[0].IsComplete());
-		CHECK(Graph->GetEdges().Contains(Edges[0]) == true);
+		VerifyEdges({ Params }, true);
+		VertexShouldOnlyHaveEdgesTo(VertexHandles[0], {VertexHandles[1]});
+		VertexShouldOnlyHaveEdgesTo(VertexHandles[1], {VertexHandles[0]});
 	}
 
 	{
 		FinalizeVertices();
-		CHECK(Graph->NumEdges() == 1);
 		CHECK(Graph->NumIslands() == 3);
 	}
 
 	{
-		FEdgeCreationParameters Params;
-		Params.VertexHandle1 = VertexHandles[2];
-		Params.VertexHandle2 = VertexHandles[3];
-		Params.EdgeIndex = FGraphUniqueIndex{ FGuid{ 2, 2, 2, 2 } };
-
-		TArray<FGraphEdgeHandle> Edges;
-		Graph->CreateBulkEdges({ Params }, &Edges);
-		CHECK(Graph->NumEdges() == 2);
+		FEdgeSpecifier Params{ VertexHandles[2], VertexHandles[3] };
+		Graph->CreateBulkEdges({ Params });
 		CHECK(Graph->NumIslands() == 2);
+		VerifyEdges({ Params }, true);
 
-		REQUIRE(Edges.Num() == 1);
-		CHECK(Edges[0].GetUniqueIndex() == Params.EdgeIndex);
-		CHECK(Edges[0].IsComplete());
-		CHECK(Graph->GetEdges().Contains(Edges[0]) == true);
+		VertexShouldOnlyHaveEdgesTo(VertexHandles[2], {VertexHandles[3]});
+		VertexShouldOnlyHaveEdgesTo(VertexHandles[3], {VertexHandles[2]});
 	}
 
 	{
-		FEdgeCreationParameters Params;
-		Params.VertexHandle1 = VertexHandles[0];
-		Params.VertexHandle2 = VertexHandles[2];
-		Params.EdgeIndex = FGraphUniqueIndex{ FGuid{ 3, 3, 3, 3 } };
-
-		TArray<FGraphEdgeHandle> Edges;
-		Graph->CreateBulkEdges({ Params }, &Edges);
-		CHECK(Graph->NumEdges() == 3);
+		FEdgeSpecifier Params{ VertexHandles[0], VertexHandles[2] };
+		Graph->CreateBulkEdges({ Params });
 		CHECK(Graph->NumIslands() == 1);
+		VerifyEdges( { Params }, true);
 
-		REQUIRE(Edges.Num() == 1);
-		CHECK(Edges[0].GetUniqueIndex() == Params.EdgeIndex);
-		CHECK(Graph->GetEdges().Contains(Edges[0]) == true);
+		VertexShouldOnlyHaveEdgesTo(VertexHandles[0], {VertexHandles[1], VertexHandles[2]});
+		VertexShouldOnlyHaveEdgesTo(VertexHandles[1], {VertexHandles[0]});
+		VertexShouldOnlyHaveEdgesTo(VertexHandles[2], {VertexHandles[0], VertexHandles[3]});
+		VertexShouldOnlyHaveEdgesTo(VertexHandles[3], {VertexHandles[2]});
 	}
 
 	{
-		FEdgeCreationParameters Params;
-		Params.VertexHandle1 = VertexHandles[0];
-		Params.VertexHandle2 = VertexHandles[2];
-		Params.EdgeIndex = FGraphUniqueIndex{ FGuid{ 4, 4, 4, 4 } };
-
-		TArray<FGraphEdgeHandle> Edges;
-		Graph->CreateBulkEdges({ Params }, &Edges);
-		CHECK(Graph->NumEdges() == 3);
+		FEdgeSpecifier Params{ VertexHandles[0], VertexHandles[2] };
+		Graph->CreateBulkEdges({ Params });
 		CHECK(Graph->NumIslands() == 1);
+		VerifyEdges({ Params }, true);
 
-		REQUIRE(Edges.Num() == 0);
+		VertexShouldOnlyHaveEdgesTo(VertexHandles[0], {VertexHandles[1], VertexHandles[2]});
+		VertexShouldOnlyHaveEdgesTo(VertexHandles[1], {VertexHandles[0]});
+		VertexShouldOnlyHaveEdgesTo(VertexHandles[2], {VertexHandles[0], VertexHandles[3]});
+		VertexShouldOnlyHaveEdgesTo(VertexHandles[3], {VertexHandles[2]});
 	}
 
 	{
-		FEdgeCreationParameters Params;
-		Params.VertexHandle1 = VertexHandles[2];
-		Params.VertexHandle2 = VertexHandles[0];
-		Params.EdgeIndex = FGraphUniqueIndex{ FGuid{ 5, 5, 5, 5 } };
-
-		TArray<FGraphEdgeHandle> Edges;
-		Graph->CreateBulkEdges({ Params }, &Edges);
-		CHECK(Graph->NumEdges() == 3);
+		FEdgeSpecifier Params{ VertexHandles[2], VertexHandles[0] };
+		Graph->CreateBulkEdges({ Params });
 		CHECK(Graph->NumIslands() == 1);
+		VerifyEdges({ Params }, true);
 
-		REQUIRE(Edges.Num() == 0);
+		VertexShouldOnlyHaveEdgesTo(VertexHandles[0], {VertexHandles[1], VertexHandles[2]});
+		VertexShouldOnlyHaveEdgesTo(VertexHandles[1], {VertexHandles[0]});
+		VertexShouldOnlyHaveEdgesTo(VertexHandles[2], {VertexHandles[0], VertexHandles[3]});
+		VertexShouldOnlyHaveEdgesTo(VertexHandles[3], {VertexHandles[2]});
 	}
 }
 
@@ -247,59 +217,31 @@ TEST_CASE_METHOD(FTestGraphBuilder, "Graph::Create Edges::Bulk", "[graph]")
 	PopulateVertices(4, true);
 
 	{
-		FEdgeCreationParameters Params1;
-		Params1.VertexHandle1 = VertexHandles[1];
-		Params1.VertexHandle2 = VertexHandles[0];
-		Params1.EdgeIndex = FGraphUniqueIndex{ FGuid{ 1, 1, 1, 1 } };
-
-		FEdgeCreationParameters Params2;
-		Params2.VertexHandle1 = VertexHandles[2];
-		Params2.VertexHandle2 = VertexHandles[3];
-		Params2.EdgeIndex = FGraphUniqueIndex{ FGuid{ 2, 2, 2, 2 } };
+		FEdgeSpecifier Params1{ VertexHandles[1], VertexHandles[0] };
+		FEdgeSpecifier Params2{ VertexHandles[2], VertexHandles[3] };
+		FEdgeSpecifier Params3{ VertexHandles[0], VertexHandles[2] };
+		Graph->CreateBulkEdges({ Params1, Params2, Params3 });
+		CHECK(Graph->NumIslands() == 1);
 
-		FEdgeCreationParameters Params3;
-		Params3.VertexHandle1 = VertexHandles[0];
-		Params3.VertexHandle2 = VertexHandles[2];
-		Params3.EdgeIndex = FGraphUniqueIndex{ FGuid{ 3, 3, 3, 3 } };
+		VerifyEdges({ Params1, Params2, Params3 }, true);
+		VertexShouldOnlyHaveEdgesTo(VertexHandles[0], {VertexHandles[1], VertexHandles[2]});
+		VertexShouldOnlyHaveEdgesTo(VertexHandles[1], {VertexHandles[0]});
+		VertexShouldOnlyHaveEdgesTo(VertexHandles[2], {VertexHandles[0], VertexHandles[3]});
+		VertexShouldOnlyHaveEdgesTo(VertexHandles[3], {VertexHandles[2]});
+	}
 
-		TArray<FGraphEdgeHandle> Edges;
-		Graph->CreateBulkEdges({ Params1, Params2, Params3 }, &Edges);
-		CHECK(Graph->NumEdges() == 3);
+	{
+		FEdgeSpecifier Params1{ VertexHandles[0], VertexHandles[1] };
+		FEdgeSpecifier Params2{ VertexHandles[3], VertexHandles[2] };
+		FEdgeSpecifier Params3{ VertexHandles[2], VertexHandles[0] };
+		Graph->CreateBulkEdges({ Params1, Params2, Params3 });
 		CHECK(Graph->NumIslands() == 1);
 
-		REQUIRE(Edges.Num() == 3);
-		CHECK(Edges[0].GetUniqueIndex() == Params1.EdgeIndex);
-		CHECK(Edges[1].GetUniqueIndex() == Params2.EdgeIndex);
-		CHECK(Edges[2].GetUniqueIndex() == Params3.EdgeIndex);
-		CHECK(Graph->GetEdges().Contains(Edges[0]) == true);
-		CHECK(Graph->GetEdges().Contains(Edges[1]) == true);
-		CHECK(Graph->GetEdges().Contains(Edges[2]) == true);
-		CHECK(Edges[0].IsComplete() == true);
-		CHECK(Edges[1].IsComplete() == true);
-		CHECK(Edges[2].IsComplete() == true);
-	}
-
-	{
-		FEdgeCreationParameters Params1;
-		Params1.VertexHandle1 = VertexHandles[0];
-		Params1.VertexHandle2 = VertexHandles[1];
-		Params1.EdgeIndex = FGraphUniqueIndex{ FGuid{ 4, 4, 4, 4 } };
-
-		FEdgeCreationParameters Params2;
-		Params2.VertexHandle1 = VertexHandles[3];
-		Params2.VertexHandle2 = VertexHandles[2];
-		Params2.EdgeIndex = FGraphUniqueIndex{ FGuid{ 5, 5, 5, 5 } };
-
-		FEdgeCreationParameters Params3;
-		Params3.VertexHandle1 = VertexHandles[2];
-		Params3.VertexHandle2 = VertexHandles[0];
-		Params3.EdgeIndex = FGraphUniqueIndex{ FGuid{ 6, 6, 6, 6 } };
-
-		TArray<FGraphEdgeHandle> Edges;
-		Graph->CreateBulkEdges({ Params1, Params2, Params3 }, &Edges);
-		CHECK(Graph->NumEdges() == 3);
-		CHECK(Graph->NumIslands() == 1);
-		CHECK(Edges.Num() == 0);
+		VerifyEdges({ Params1, Params2, Params3 }, true);
+		VertexShouldOnlyHaveEdgesTo(VertexHandles[0], {VertexHandles[1], VertexHandles[2]});
+		VertexShouldOnlyHaveEdgesTo(VertexHandles[1], {VertexHandles[0]});
+		VertexShouldOnlyHaveEdgesTo(VertexHandles[2], {VertexHandles[0], VertexHandles[3]});
+		VertexShouldOnlyHaveEdgesTo(VertexHandles[3], {VertexHandles[2]});
 	}
 }
 
@@ -318,11 +260,9 @@ public:
 			REQUIRE(Vertex != nullptr);
 
 			Vertex->ForEachAdjacentVertex(
-				[this](const FGraphVertexHandle& NeighborHandle, const FGraphEdgeHandle& EdgeHandle)
+				[this, &VertexHandle](const FGraphVertexHandle& NeighborHandle)
 				{
-					EdgeHandles.Add(EdgeHandle);
-					CHECK(Graph->GetEdges().Contains(EdgeHandle) == true);
-					CHECK(EdgeHandle.IsComplete() == true);
+					EdgeHandles.Add(FEdgeSpecifier{VertexHandle, NeighborHandle});
 				}
 			);
 		}
@@ -330,16 +270,23 @@ public:
 
 	~FScopedVertexEdgeCleanupChecker()
 	{
-		for (const FGraphEdgeHandle& EdgeHandle : EdgeHandles)
+		for (const FEdgeSpecifier& EdgeHandle : EdgeHandles)
 		{
-			CHECK(Graph->GetEdges().Contains(EdgeHandle) == false);
-			CHECK(EdgeHandle.IsComplete() == false);
+			if (UGraphVertex* V1 = EdgeHandle.GetVertexHandle1().GetVertex())
+			{
+				CHECK(V1->HasEdgeTo(EdgeHandle.GetVertexHandle2()) == false);
+			}
+
+			if (UGraphVertex* V2 = EdgeHandle.GetVertexHandle2().GetVertex())
+			{
+				CHECK(V2->HasEdgeTo(EdgeHandle.GetVertexHandle1()) == false);
+			}
 		}
 	}
 
 private:
 	UGraph* Graph;
-	TArray<FGraphEdgeHandle> EdgeHandles;
+	TArray<FEdgeSpecifier> EdgeHandles;
 };
 
 class FScopedVertexIslandCleanupChecker
@@ -383,7 +330,6 @@ TEST_CASE_METHOD(FTestGraphBuilder, "Graph::Remove Vertex::Single::Fully Connect
 	BuildFullyConnectedEdges(5);
 
 	CHECK(Graph->NumVertices() == 5);
-	CHECK(Graph->NumEdges() == 10);
 	CHECK(Graph->NumIslands() == 1);
 
 	{
@@ -395,7 +341,6 @@ TEST_CASE_METHOD(FTestGraphBuilder, "Graph::Remove Vertex::Single::Fully Connect
 	CHECK(Graph->GetCompleteNodeHandle(VertexHandles[0]).IsComplete() == false);
 	CHECK(Graph->GetVertices().Contains(VertexHandles[0]) == false);
 	CHECK(Graph->NumVertices() == 4);
-	CHECK(Graph->NumEdges() == 6);
 	CHECK(Graph->NumIslands() == 1);
 
 	{
@@ -407,7 +352,6 @@ TEST_CASE_METHOD(FTestGraphBuilder, "Graph::Remove Vertex::Single::Fully Connect
 	CHECK(Graph->GetCompleteNodeHandle(VertexHandles[1]).IsComplete() == false);
 	CHECK(Graph->GetVertices().Contains(VertexHandles[1]) == false);
 	CHECK(Graph->NumVertices() == 3);
-	CHECK(Graph->NumEdges() == 3);
 	CHECK(Graph->NumIslands() == 1);
 
 	{
@@ -419,7 +363,6 @@ TEST_CASE_METHOD(FTestGraphBuilder, "Graph::Remove Vertex::Single::Fully Connect
 	CHECK(Graph->GetCompleteNodeHandle(VertexHandles[2]).IsComplete() == false);
 	CHECK(Graph->GetVertices().Contains(VertexHandles[2]) == false);
 	CHECK(Graph->NumVertices() == 2);
-	CHECK(Graph->NumEdges() == 1);
 	CHECK(Graph->NumIslands() == 1);
 
 	{
@@ -431,7 +374,6 @@ TEST_CASE_METHOD(FTestGraphBuilder, "Graph::Remove Vertex::Single::Fully Connect
 	CHECK(Graph->GetCompleteNodeHandle(VertexHandles[3]).IsComplete() == false);
 	CHECK(Graph->GetVertices().Contains(VertexHandles[3]) == false);
 	CHECK(Graph->NumVertices() == 1);
-	CHECK(Graph->NumEdges() == 0);
 	CHECK(Graph->NumIslands() == 1);
 
 	{
@@ -443,7 +385,6 @@ TEST_CASE_METHOD(FTestGraphBuilder, "Graph::Remove Vertex::Single::Fully Connect
 	CHECK(Graph->GetCompleteNodeHandle(VertexHandles[4]).IsComplete() == false);
 	CHECK(Graph->GetVertices().Contains(VertexHandles[4]) == false);
 	CHECK(Graph->NumVertices() == 0);
-	CHECK(Graph->NumEdges() == 0);
 	CHECK(Graph->NumIslands() == 0);
 }
 
@@ -453,7 +394,6 @@ TEST_CASE_METHOD(FTestGraphBuilder, "Graph::Remove Vertex::Single::Linear", "[gr
 	BuildLinearEdges(10);
 
 	CHECK(Graph->NumVertices() == 10);
-	CHECK(Graph->NumEdges() == 9);
 	CHECK(Graph->NumIslands() == 1);
 
 	SECTION("First Vertex")
@@ -468,7 +408,6 @@ TEST_CASE_METHOD(FTestGraphBuilder, "Graph::Remove Vertex::Single::Linear", "[gr
 		CHECK(VertexHandles[0].IsComplete() == false);
 
 		CHECK(Graph->NumVertices() == 9);
-		CHECK(Graph->NumEdges() == 8);
 		CHECK(Graph->NumIslands() == 1);
 	}
 
@@ -484,7 +423,6 @@ TEST_CASE_METHOD(FTestGraphBuilder, "Graph::Remove Vertex::Single::Linear", "[gr
 		CHECK(VertexHandles[5].IsComplete() == false);
 
 		CHECK(Graph->NumVertices() == 9);
-		CHECK(Graph->NumEdges() == 7);
 		CHECK(Graph->NumIslands() == 2);
 	}
 }
@@ -495,7 +433,6 @@ TEST_CASE_METHOD(FTestGraphBuilder, "Graph::Remove Vertex::Bulk::Fully Connected
 	BuildFullyConnectedEdges(5);
 
 	CHECK(Graph->NumVertices() == 5);
-	CHECK(Graph->NumEdges() == 10);
 	CHECK(Graph->NumIslands() == 1);
 
 	{
@@ -508,7 +445,6 @@ TEST_CASE_METHOD(FTestGraphBuilder, "Graph::Remove Vertex::Bulk::Fully Connected
 	CHECK(VertexHandles[0].IsComplete() == false);
 	CHECK(VertexHandles[1].IsComplete() == false);
 	CHECK(Graph->NumVertices() == 3);
-	CHECK(Graph->NumEdges() == 3);
 	CHECK(Graph->NumIslands() == 1);
 
 	{
@@ -523,7 +459,6 @@ TEST_CASE_METHOD(FTestGraphBuilder, "Graph::Remove Vertex::Bulk::Fully Connected
 	CHECK(VertexHandles[3].IsComplete() == false);
 	CHECK(VertexHandles[4].IsComplete() == false);
 	CHECK(Graph->NumVertices() == 0);
-	CHECK(Graph->NumEdges() == 0);
 	CHECK(Graph->NumIslands() == 0);
 }
 
@@ -534,7 +469,6 @@ TEST_CASE_METHOD(FTestGraphBuilder, "Graph::Remove Vertex::Bulk::Linear", "[grap
 	BuildLinearEdges(10);
 
 	CHECK(Graph->NumVertices() == 10);
-	CHECK(Graph->NumEdges() == 9);
 	CHECK(Graph->NumIslands() == 1);
 
 	TArray<FGraphVertexHandle> VerticesToRemove = { VertexHandles[1], VertexHandles[2], VertexHandles[8] };
@@ -553,7 +487,6 @@ TEST_CASE_METHOD(FTestGraphBuilder, "Graph::Remove Vertex::Bulk::Linear", "[grap
 	}
 
 	CHECK(Graph->NumVertices() == 7);
-	CHECK(Graph->NumEdges() == 4);
 	CHECK(Graph->NumIslands() == 3);
 }
 
@@ -564,7 +497,6 @@ TEST_CASE_METHOD(FTestGraphBuilder, "Graph::Remove Vertex::Incomplete Handle", "
 
 	FGraphVertexHandle RemoveHandle{ VertexHandles[0].GetUniqueIndex(), nullptr };
 	CHECK(Graph->NumVertices() == 5);
-	CHECK(Graph->NumEdges() == 4);
 	CHECK(Graph->NumIslands() == 1);
 
 	UGraphIsland* Island = IslandHandles[0].GetIsland();
@@ -574,7 +506,6 @@ TEST_CASE_METHOD(FTestGraphBuilder, "Graph::Remove Vertex::Incomplete Handle", "
 	Graph->RemoveVertex(RemoveHandle);
 
 	CHECK(Graph->NumVertices() == 4);
-	CHECK(Graph->NumEdges() == 3);
 	CHECK(Graph->NumIslands() == 1);
 	CHECK(Island->Num() == 4);
 }
@@ -585,7 +516,6 @@ TEST_CASE_METHOD(FTestGraphBuilder, "Graph::Remove Island", "[graph]")
 	BuildFullyConnectedEdges(5);
 
 	CHECK(Graph->NumVertices() == 10);
-	CHECK(Graph->NumEdges() == 20);
 	CHECK(Graph->NumIslands() == 2);
 
 	{
@@ -611,7 +541,6 @@ TEST_CASE_METHOD(FTestGraphBuilder, "Graph::Remove Island", "[graph]")
 	}
 
 	CHECK(Graph->NumVertices() == 5);
-	CHECK(Graph->NumEdges() == 10);
 	CHECK(Graph->NumIslands() == 1);
 
 	{
@@ -637,58 +566,9 @@ TEST_CASE_METHOD(FTestGraphBuilder, "Graph::Remove Island", "[graph]")
 	}
 
 	CHECK(Graph->NumVertices() == 0);
-	CHECK(Graph->NumEdges() == 0);
 	CHECK(Graph->NumIslands() == 0);
 }
 
-TEST_CASE_METHOD(FTestGraphBuilder, "Graph::Remove Edge::Handle Islands", "[graph]")
-{
-	PopulateVertices(10, true);
-	BuildLinearEdges(10);
-
-	CHECK(Graph->NumVertices() == 10);
-	CHECK(Graph->NumEdges() == 9);
-	CHECK(Graph->NumIslands() == 1);
-
-	SECTION("First Vertex")
-	{
-		TArray<FGraphEdgeHandle> Edges = GetEdgesForVertex(VertexHandles[0]);
-		for (const FGraphEdgeHandle& EdgeHandle : Edges)
-		{
-			CHECK(Graph->GetEdges().Contains(EdgeHandle) == true);
-			CHECK(EdgeHandle.IsComplete() == true);
-
-			Graph->RemoveEdge(EdgeHandle);
-
-			CHECK(Graph->GetEdges().Contains(EdgeHandle) == false);
-			CHECK(EdgeHandle.IsComplete() == false);
-		}
-
-		CHECK(Graph->NumVertices() == 10);
-		CHECK(Graph->NumEdges() == 8);
-		CHECK(Graph->NumIslands() == 2);
-	}
-
-	SECTION("Split Island")
-	{
-		TArray<FGraphEdgeHandle> Edges = GetEdgesForVertex(VertexHandles[5]);
-		for (const FGraphEdgeHandle& EdgeHandle : Edges)
-		{
-			CHECK(Graph->GetEdges().Contains(EdgeHandle) == true);
-			CHECK(EdgeHandle.IsComplete() == true);
-
-			Graph->RemoveEdge(EdgeHandle);
-
-			CHECK(Graph->GetEdges().Contains(EdgeHandle) == false);
-			CHECK(EdgeHandle.IsComplete() == false);
-		}
-
-		CHECK(Graph->NumVertices() == 10);
-		CHECK(Graph->NumEdges() == 7);
-		CHECK(Graph->NumIslands() == 3);
-	}
-}
-
 TEST_CASE_METHOD(FTestGraphBuilder, "Graph::Event::Vertex Created", "[graph]")
 {
 	FGraphUniqueIndex Node1Index = { FGuid { 1, 1, 1, 1 } };
@@ -708,40 +588,11 @@ TEST_CASE_METHOD(FTestGraphBuilder, "Graph::Event::Vertex Created", "[graph]")
 	CHECK(CreatedHandle.IsComplete() == true);
 }
 
-TEST_CASE_METHOD(FTestGraphBuilder, "Graph::Event::Edge Created::Single", "[graph]")
-{
-	PopulateVertices(10, true);
-
-	FEdgeCreationParameters Params;
-	Params.VertexHandle1 = VertexHandles[1];
-	Params.VertexHandle2 = VertexHandles[0];
-	Params.EdgeIndex = FGraphUniqueIndex{ FGuid{ 1, 1, 1, 1 } };
-
-	FGraphEdgeHandle CreatedHandle;
-	Graph->OnEdgeCreated.AddLambda(
-		[&CreatedHandle, &Params](const FGraphEdgeHandle& InCreatedHandle)
-		{
-			CHECK(InCreatedHandle.IsComplete());
-			CHECK(InCreatedHandle.GetUniqueIndex() == Params.EdgeIndex);
-			CreatedHandle = InCreatedHandle;
-		}
-	);
-
-	TArray<FGraphEdgeHandle> Edges;
-	Graph->CreateBulkEdges({ Params }, &Edges);
-
-	REQUIRE(Edges.Num() == 1);
-	CHECK(Edges[0] == CreatedHandle);
-}
-
 TEST_CASE_METHOD(FTestGraphBuilder, "Graph::Event::Island Created", "[graph]")
 {
 	PopulateVertices(10, false);
 
-	FEdgeCreationParameters Params;
-	Params.VertexHandle1 = VertexHandles[1];
-	Params.VertexHandle2 = VertexHandles[0];
-	Params.EdgeIndex = FGraphUniqueIndex{ FGuid{ 1, 1, 1, 1 } };
+	FEdgeSpecifier Params{ VertexHandles[1], VertexHandles[0] };
 
 	FGraphIslandHandle CreatedHandle;
 	Graph->OnIslandCreated.AddLambda(
diff --git a/Engine/Plugins/Experimental/GameplayGraph/Tests/Private/GraphVertexTests.cpp b/Engine/Plugins/Experimental/GameplayGraph/Tests/Private/GraphVertexTests.cpp
--- a/Engine/Plugins/Experimental/GameplayGraph/Tests/Private/GraphVertexTests.cpp
+++ b/Engine/Plugins/Experimental/GameplayGraph/Tests/Private/GraphVertexTests.cpp
@@ -73,148 +73,6 @@ TEST_CASE_METHOD(FTestGraphBuilder, "Graph::Vertex::Has Edge To", "[graph][verte
 			CHECK(Vertex1->GetEdges().Contains(VertexHandles[0]) == false);
 		}
 	}
-
-	SECTION("Post remove edge")
-	{
-		{
-			UGraphVertex* Vertex0 = VertexHandles[0].GetVertex();
-			REQUIRE(Vertex0 != nullptr);
-			UGraphVertex* Vertex1 = VertexHandles[1].GetVertex();
-			REQUIRE(Vertex1 != nullptr);
-
-			const FGraphEdgeHandle* Edge = Vertex0->FindEdgeTo(VertexHandles[1]);
-			REQUIRE(Edge != nullptr);
-
-			CHECK(Edge->IsComplete() == true);
-			Graph->RemoveEdge(*Edge);
-			CHECK(Edge->IsComplete() == false);
-
-			CHECK(Vertex0->HasEdgeTo(VertexHandles[1]) == false);
-			CHECK(Vertex1->HasEdgeTo(VertexHandles[0]) == false);
-		}
-
-		{
-			for (int32 Index1 = 0; Index1 < NumVertices; ++Index1)
-			{
-				UGraphVertex* Vertex1 = VertexHandles[Index1].GetVertex();
-				REQUIRE(Vertex1 != nullptr);
-
-				for (int32 Index2 = Index1 + 1; Index2 < NumVertices; ++Index2)
-				{
-					if (Index1 == 0 && Index2 == 1)
-					{
-						continue;
-					}
-
-					UGraphVertex* Vertex2 = VertexHandles[Index2].GetVertex();
-					REQUIRE(Vertex2 != nullptr);
-
-					CHECK(Vertex1->HasEdgeTo(VertexHandles[Index2]) == true);
-					CHECK(Vertex1->GetEdges().Contains(VertexHandles[Index2]) == true);
-
-					CHECK(Vertex2->HasEdgeTo(VertexHandles[Index1]) == true);
-					CHECK(Vertex2->GetEdges().Contains(VertexHandles[Index1]) == true);
-				}
-			}
-		}
-	}
-}
-
-TEST_CASE_METHOD(FTestGraphBuilder, "Graph::Vertex::Find Edge To", "[graph][vertex]")
-{
-	constexpr int32 NumVertices = 7;
-	PopulateVertices(NumVertices, true);
-	BuildFullyConnectedEdges(NumVertices);
-
-	SECTION("Post Construction")
-	{
-		for (int32 Index1 = 0; Index1 < NumVertices; ++Index1)
-		{
-			UGraphVertex* Vertex1 = VertexHandles[Index1].GetVertex();
-			REQUIRE(Vertex1 != nullptr);
-
-			for (int32 Index2 = Index1 + 1; Index2 < NumVertices; ++Index2)
-			{
-				UGraphVertex* Vertex2 = VertexHandles[Index2].GetVertex();
-				REQUIRE(Vertex2 != nullptr);
-			
-				const FGraphEdgeHandle* Edge1 = Vertex1->FindEdgeTo(VertexHandles[Index2]);
-				REQUIRE(Edge1 != nullptr);
-				const FGraphEdgeHandle* Edge2 = Vertex2->FindEdgeTo(VertexHandles[Index1]);
-				REQUIRE(Edge2 != nullptr);
-				CHECK(*Edge1 == *Edge2);
-			}
-		}
-	}
-
-	SECTION("Post remove vertex")
-	{
-		Graph->RemoveVertex(VertexHandles[0]);
-		for (int32 Index1 = 1; Index1 < NumVertices; ++Index1)
-		{
-			UGraphVertex* Vertex1 = VertexHandles[Index1].GetVertex();
-			REQUIRE(Vertex1 != nullptr);
-
-			for (int32 Index2 = Index1 + 1; Index2 < NumVertices; ++Index2)
-			{
-				UGraphVertex* Vertex2 = VertexHandles[Index2].GetVertex();
-				REQUIRE(Vertex2 != nullptr);
-
-				const FGraphEdgeHandle* Edge1 = Vertex1->FindEdgeTo(VertexHandles[Index2]);
-				REQUIRE(Edge1 != nullptr);
-				const FGraphEdgeHandle* Edge2 = Vertex2->FindEdgeTo(VertexHandles[Index1]);
-				REQUIRE(Edge2 != nullptr);
-				CHECK(*Edge1 == *Edge2);
-			}
-
-			CHECK(Vertex1->FindEdgeTo(VertexHandles[0]) == nullptr);
-		}
-	}
-
-	SECTION("Post remove edge")
-	{
-		{
-			UGraphVertex* Vertex0 = VertexHandles[0].GetVertex();
-			REQUIRE(Vertex0 != nullptr);
-			UGraphVertex* Vertex1 = VertexHandles[1].GetVertex();
-			REQUIRE(Vertex1 != nullptr);
-
-			const FGraphEdgeHandle* Edge = Vertex0->FindEdgeTo(VertexHandles[1]);
-			REQUIRE(Edge != nullptr);
-
-			CHECK(Edge->IsComplete() == true);
-			Graph->RemoveEdge(*Edge);
-			CHECK(Edge->IsComplete() == false);
-
-			CHECK(Vertex0->FindEdgeTo(VertexHandles[1]) == nullptr);
-			CHECK(Vertex1->FindEdgeTo(VertexHandles[0]) == nullptr);
-		}
-
-		{
-			for (int32 Index1 = 0; Index1 < NumVertices; ++Index1)
-			{
-				UGraphVertex* Vertex1 = VertexHandles[Index1].GetVertex();
-				REQUIRE(Vertex1 != nullptr);
-
-				for (int32 Index2 = Index1 + 1; Index2 < NumVertices; ++Index2)
-				{
-					if (Index1 == 0 && Index2 == 1)
-					{
-						continue;
-					}
-
-					UGraphVertex* Vertex2 = VertexHandles[Index2].GetVertex();
-					REQUIRE(Vertex2 != nullptr);
-
-					const FGraphEdgeHandle* Edge1 = Vertex1->FindEdgeTo(VertexHandles[Index2]);
-					REQUIRE(Edge1 != nullptr);
-					const FGraphEdgeHandle* Edge2 = Vertex2->FindEdgeTo(VertexHandles[Index1]);
-					REQUIRE(Edge2 != nullptr);
-					CHECK(*Edge1 == *Edge2);
-				}
-			}
-		}
-	}
 }
 
 TEST_CASE_METHOD(FTestGraphBuilder, "Graph::Vertex::Num Edges", "[graph][vertex]")
@@ -243,34 +101,6 @@ TEST_CASE_METHOD(FTestGraphBuilder, "Graph::Vertex::Num Edges", "[graph][vertex]
 			CHECK(Vertex1->NumEdges() == 5);
 		}
 	}
-
-	SECTION("Post remove edge")
-	{
-		UGraphVertex* Vertex0 = VertexHandles[0].GetVertex();
-		REQUIRE(Vertex0 != nullptr);
-
-		const FGraphEdgeHandle* Edge = Vertex0->FindEdgeTo(VertexHandles[1]);
-		REQUIRE(Edge != nullptr);
-
-		CHECK(Edge->IsComplete() == true);
-		Graph->RemoveEdge(*Edge);
-		CHECK(Edge->IsComplete() == false);
-
-		for (int32 Index1 = 0; Index1 < NumVertices; ++Index1)
-		{
-			UGraphVertex* Vertex1 = VertexHandles[Index1].GetVertex();
-			REQUIRE(Vertex1 != nullptr);
-
-			if (Index1 == 0 || Index1 == 1)
-			{
-				CHECK(Vertex1->NumEdges() == 5);
-			}
-			else
-			{
-				CHECK(Vertex1->NumEdges() == 6);
-			}
-		}
-	}
 }
 
 TEST_CASE_METHOD(FTestGraphBuilder, "Graph::Vertex::Get Parent Island", "[graph][vertex]")
@@ -321,22 +151,6 @@ TEST_CASE_METHOD(FTestGraphBuilder, "Graph::Vertex::Get Parent Island", "[graph]
 			CHECK(Vertex0->GetParentIsland() == Vertex1->GetParentIsland());
 		}
 	}
-
-	SECTION("Post remove edge")
-	{
-		UGraphVertex* Vertex0 = VertexHandles[3].GetVertex();
-		REQUIRE(Vertex0 != nullptr);
-
-		const FGraphEdgeHandle* Edge = Vertex0->FindEdgeTo(VertexHandles[4]);
-		REQUIRE(Edge != nullptr);
-
-		CHECK(Edge->IsComplete() == true);
-		Graph->RemoveEdge(*Edge);
-		CHECK(Edge->IsComplete() == false);
-
-		CHECK(Vertex0->GetParentIsland() != VertexHandles[4].GetVertex()->GetParentIsland());
-		CHECK(Vertex0->GetParentIsland() == VertexHandles[2].GetVertex()->GetParentIsland());
-	}
 }
 
 TEST_CASE_METHOD(FTestGraphBuilder, "Graph::Vertex::For Each Adjacent::With Edge", "[graph][vertex]")
@@ -350,17 +164,22 @@ TEST_CASE_METHOD(FTestGraphBuilder, "Graph::Vertex::For Each Adjacent::With Edge
 		UGraphVertex* Vertex1 = VertexHandles[Index1].GetVertex();
 		REQUIRE(Vertex1 != nullptr);
 
-		TMap<FGraphVertexHandle, FGraphEdgeHandle> ConnectedVertices;
+		TSet<FGraphVertexHandle> ConnectedVertices;
 		ConnectedVertices.Reserve(NumVertices - 1);
 
 		Vertex1->ForEachAdjacentVertex(
-			[&ConnectedVertices](const FGraphVertexHandle& NeighborHandle, const FGraphEdgeHandle& EdgeHandle)
+			[&ConnectedVertices](const FGraphVertexHandle& NeighborHandle)
 			{
-				ConnectedVertices.Add(NeighborHandle, EdgeHandle);
+				ConnectedVertices.Add(NeighborHandle);
 			}
 		);
 
-		CHECK(ConnectedVertices == Vertex1->GetEdges());
+		CHECK(ConnectedVertices.Num() == Vertex1->GetEdges().Num());
+		for (const FGraphVertexHandle& Adjacent : Vertex1->GetEdges())
+		{
+			CHECK(ConnectedVertices.Contains(Adjacent) == true);
+		}
+
 		CHECK(ConnectedVertices.Num() == NumVertices - 1);
 		for (int32 Index2 = 0; Index2 < NumVertices; ++Index2)
 		{
@@ -369,10 +188,8 @@ TEST_CASE_METHOD(FTestGraphBuilder, "Graph::Vertex::For Each Adjacent::With Edge
 				continue;
 			}
 			
-			CHECK(ConnectedVertices.Contains(VertexHandles[Index2]));
-			FGraphEdgeHandle* EdgeHandle = ConnectedVertices.Find(VertexHandles[Index2]);
-			REQUIRE(EdgeHandle != nullptr);
-			CHECK(Vertex1->GetEdges().FindRef(VertexHandles[Index2]) == *EdgeHandle);
+			CHECK(ConnectedVertices.Contains(VertexHandles[Index2]) == true);
+			CHECK(Vertex1->GetEdges().Contains(VertexHandles[Index2]) == true);
 			ConnectedVertices.Remove(VertexHandles[Index2]);
 		}
 		CHECK(ConnectedVertices.Num() == 0);
@@ -485,13 +302,9 @@ TEST_CASE_METHOD(FTestGraphBuilder, "Graph::Vertex::Event::Parent Island Set", "
 
 	for (Index = 1; Index < NumVertices; ++Index)
 	{
-		FEdgeCreationParameters Params;
-		Params.VertexHandle1 = VertexHandles[Index-1];
-		Params.VertexHandle2 = VertexHandles[Index];
-
-		TArray<FGraphEdgeHandle> Edges;
-		Graph->CreateBulkEdges({ Params }, &Edges);
-		REQUIRE(Edges.Num() == 1);
+		FEdgeSpecifier Params{ VertexHandles[Index - 1], VertexHandles[Index] };
+		Graph->CreateBulkEdges({ Params });
+		VerifyEdges({ Params }, true);
 
 		for (int32 Index2 = 0; Index2 < NumVertices; ++Index2)
 		{
diff --git a/Engine/Plugins/Experimental/GameplayGraph/Tests/Private/TestGraphBuilder.cpp b/Engine/Plugins/Experimental/GameplayGraph/Tests/Private/TestGraphBuilder.cpp
--- a/Engine/Plugins/Experimental/GameplayGraph/Tests/Private/TestGraphBuilder.cpp
+++ b/Engine/Plugins/Experimental/GameplayGraph/Tests/Private/TestGraphBuilder.cpp
@@ -41,7 +41,7 @@ void FTestGraphBuilder::BuildFullyConnectedEdges(int32 NodesPerIsland)
 {
 	for (int32 NodeIndex = 0; NodeIndex < VertexHandles.Num(); NodeIndex += NodesPerIsland)
 	{
-		TArray<FEdgeCreationParameters> AllEdges;
+		TArray<FEdgeSpecifier> AllEdges;
 		for (int32 SourceOffset = 0; SourceOffset < NodesPerIsland; ++SourceOffset)
 		{
 			if ((NodeIndex + SourceOffset) >= VertexHandles.Num())
@@ -56,14 +56,14 @@ void FTestGraphBuilder::BuildFullyConnectedEdges(int32 NodesPerIsland)
 					break;
 				}
 
-				FEdgeCreationParameters Params;
-				Params.VertexHandle1 = VertexHandles[NodeIndex + SourceOffset];
-				Params.VertexHandle2 = VertexHandles[NodeIndex + DestOffset];
+				FEdgeSpecifier Params{ VertexHandles[NodeIndex + SourceOffset], VertexHandles[NodeIndex + DestOffset] };
 				AllEdges.Add(Params);
 			}
 		}
 
+		TArray<FEdgeSpecifier> AllVerifyEdges = AllEdges;
 		Graph->CreateBulkEdges(MoveTemp(AllEdges));
+		VerifyEdges(AllVerifyEdges, true);
 	}
 
 	FinalizeEdges();
@@ -73,7 +73,7 @@ void FTestGraphBuilder::BuildLinearEdges(int32 NodesPerIsland)
 {
 	for (int32 NodeIndex = 0; NodeIndex < VertexHandles.Num(); NodeIndex += NodesPerIsland)
 	{
-		TArray<FEdgeCreationParameters> AllEdges;
+		TArray<FEdgeSpecifier> AllEdges;
 		for (int32 SourceOffset = 1; SourceOffset < NodesPerIsland; ++SourceOffset)
 		{
 			if ((NodeIndex + SourceOffset) >= VertexHandles.Num())
@@ -81,13 +81,13 @@ void FTestGraphBuilder::BuildLinearEdges(int32 NodesPerIsland)
 				break;
 			}
 
-			FEdgeCreationParameters Params;
-			Params.VertexHandle1 = VertexHandles[NodeIndex + SourceOffset - 1];
-			Params.VertexHandle2 = VertexHandles[NodeIndex + SourceOffset];
+			FEdgeSpecifier Params{VertexHandles[NodeIndex + SourceOffset - 1], VertexHandles[NodeIndex + SourceOffset]};
 			AllEdges.Add(Params);
 		}
 
+		TArray<FEdgeSpecifier> AllVerifyEdges = AllEdges;
 		Graph->CreateBulkEdges(MoveTemp(AllEdges));
+		VerifyEdges(AllVerifyEdges, true);
 	}
 
 	FinalizeEdges();
@@ -102,23 +102,23 @@ void FTestGraphBuilder::FinalizeEdges()
 	}
 }
 
-TArray<FGraphEdgeHandle> FTestGraphBuilder::GetEdgesForVertex(const FGraphVertexHandle& Handle) const
+TArray<FEdgeSpecifier> FTestGraphBuilder::GetEdgesForVertex(const FGraphVertexHandle& Handle) const
 {
-	TArray<FGraphEdgeHandle> Edges;
+	TArray<FEdgeSpecifier> Edges;
 	if (UGraphVertex* Vertex = Handle.GetVertex())
 	{
 		Edges.Reserve(Vertex->NumEdges());
 		Vertex->ForEachAdjacentVertex(
-			[&Edges](const FGraphVertexHandle& NeighborVertexHandle, const FGraphEdgeHandle& EdgeHandle)
+			[&Edges, &Handle](const FGraphVertexHandle& NeighborVertexHandle)
 			{
-				Edges.Add(EdgeHandle);
+				Edges.Add(FEdgeSpecifier{Handle, NeighborVertexHandle});
 			}
 		);
 	}
 	return Edges;
 }
 
-void FTestGraphBuilder::GraphSanityCheck()
+void FTestGraphBuilder::GraphSanityCheck() const
 {
 	for (const TPair<FGraphVertexHandle, TObjectPtr<UGraphVertex>>& Data : Graph->GetVertices())
 	{
@@ -136,33 +136,20 @@ void FTestGraphBuilder::GraphSanityCheck()
 		}
 
 		Data.Value->ForEachAdjacentVertex(
-			[this](const FGraphVertexHandle& NeighborVertexHandle, const FGraphEdgeHandle& EdgeHandle)
+			[this, &Data](const FGraphVertexHandle& NeighborVertexHandle)
 			{
 				CHECK(NeighborVertexHandle.IsComplete() == true);
 				CHECK(NeighborVertexHandle.GetVertex() == Graph->GetVertices().FindRef(NeighborVertexHandle));
 
-				CHECK(EdgeHandle.IsComplete() == true);
-				CHECK(EdgeHandle.GetEdge() == Graph->GetEdges().FindRef(EdgeHandle));
+				const UGraphVertex* AdjacentVertex = NeighborVertexHandle.GetVertex();
+				REQUIRE(AdjacentVertex != nullptr);
+
+				CHECK(Data.Value->HasEdgeTo(NeighborVertexHandle) == true);
+				CHECK(AdjacentVertex->HasEdgeTo(Data.Key) == true);
 			}
 		);
 	}
 
-	for (const TPair<FGraphEdgeHandle, TObjectPtr<UGraphEdge>>& Data : Graph->GetEdges())
-	{
-		CHECK(Data.Key.IsComplete() == true);
-		CHECK(Data.Value != nullptr);
-		CHECK(Data.Value == Data.Key.GetEdge());
-		CHECK(Data.Key == Data.Value->Handle());
-		CHECK(Data.Key.GetUniqueIndex().IsTemporary() == false);
-		CHECK(Data.Key.GetGraph() == Graph);
-
-		CHECK(Data.Value->NodeA().IsComplete() == true);
-		CHECK(Data.Value->NodeA().GetVertex() == Graph->GetVertices().FindRef(Data.Value->NodeA()));
-
-		CHECK(Data.Value->NodeB().IsComplete() == true);
-		CHECK(Data.Value->NodeB().GetVertex() == Graph->GetVertices().FindRef(Data.Value->NodeB()));
-	}
-
 	for (const TPair<FGraphIslandHandle, TObjectPtr<UGraphIsland>>& Data : Graph->GetIslands())
 	{
 		CHECK(Data.Key.IsComplete() == true);
@@ -182,7 +169,7 @@ void FTestGraphBuilder::GraphSanityCheck()
 	}
 }
 
-void FTestGraphBuilder::IslandVertexParentIslandSanityCheck(const FGraphIslandHandle& IslandHandle)
+void FTestGraphBuilder::IslandVertexParentIslandSanityCheck(const FGraphIslandHandle& IslandHandle) const
 {
 	UGraphIsland* Island = IslandHandle.GetIsland();
 	REQUIRE(Island != nullptr);
@@ -194,3 +181,33 @@ void FTestGraphBuilder::IslandVertexParentIslandSanityCheck(const FGraphIslandHa
 		}
 	);
 }
+
+void FTestGraphBuilder::VerifyEdges(const TArray<FEdgeSpecifier>& Edges, bool bExist) const
+{
+	for (const FEdgeSpecifier& Edge : Edges)
+	{
+		UGraphVertex* V1 = Edge.GetVertexHandle1().GetVertex();
+		REQUIRE(V1 != nullptr);
+
+		UGraphVertex* V2 = Edge.GetVertexHandle2().GetVertex();
+		REQUIRE(V2 != nullptr);
+
+		CHECK(V1->HasEdgeTo(V2->Handle()) == bExist);
+		CHECK(V2->HasEdgeTo(V1->Handle()) == bExist);
+	}
+}
+
+void FTestGraphBuilder::VertexShouldOnlyHaveEdgesTo(const FGraphVertexHandle& Source, const TArray<FGraphVertexHandle>& Targets) const
+{
+	UGraphVertex* SourceVertex = Source.GetVertex();
+	REQUIRE(SourceVertex != nullptr);
+	CHECK(SourceVertex->NumEdges() == Targets.Num());
+	for (const FGraphVertexHandle& AdjacentHandle : Targets)
+	{
+		CHECK(SourceVertex->HasEdgeTo(AdjacentHandle) == true);
+		
+		UGraphVertex* AdjacentVertex = AdjacentHandle.GetVertex();
+		REQUIRE(AdjacentVertex != nullptr);
+		CHECK(AdjacentVertex->HasEdgeTo(Source) == true);
+	}
+}
\ No newline at end of file
diff --git a/Engine/Plugins/Experimental/GameplayGraph/Tests/Private/TestGraphBuilder.h b/Engine/Plugins/Experimental/GameplayGraph/Tests/Private/TestGraphBuilder.h
--- a/Engine/Plugins/Experimental/GameplayGraph/Tests/Private/TestGraphBuilder.h
+++ b/Engine/Plugins/Experimental/GameplayGraph/Tests/Private/TestGraphBuilder.h
@@ -17,13 +17,15 @@ public:
 
 	void FinalizeEdges();
 
-	TArray<FGraphEdgeHandle> GetEdgesForVertex(const FGraphVertexHandle& Handle) const;
+	TArray<FEdgeSpecifier> GetEdgesForVertex(const FGraphVertexHandle& Handle) const;
 
 protected:
 	TObjectPtr<UGraph> Graph;
 	TArray<FGraphVertexHandle> VertexHandles;
 	TArray<FGraphIslandHandle> IslandHandles;
 
-	void GraphSanityCheck();
-	void IslandVertexParentIslandSanityCheck(const FGraphIslandHandle& IslandHandle);
+	void GraphSanityCheck() const;
+	void IslandVertexParentIslandSanityCheck(const FGraphIslandHandle& IslandHandle) const;
+	void VerifyEdges(const TArray<FEdgeSpecifier>& Edges, bool bExist) const;
+	void VertexShouldOnlyHaveEdgesTo(const FGraphVertexHandle& Source, const TArray<FGraphVertexHandle>& Targets) const;
 };
\ No newline at end of file
