gameplay graph improvements

- deserialization should handle cases where the serialized data has a vertex in multiple islands - properly attempt to merge islands together if that's the case
- gameplay graph should not attempt to create zero-sized islands
- added tests for the above cases

diff --git a/Engine/Plugins/Experimental/GameplayGraph/Source/GameplayGraph/Private/Graph/Graph.cpp b/Engine/Plugins/Experimental/GameplayGraph/Source/GameplayGraph/Private/Graph/Graph.cpp
--- a/Engine/Plugins/Experimental/GameplayGraph/Source/GameplayGraph/Private/Graph/Graph.cpp
+++ b/Engine/Plugins/Experimental/GameplayGraph/Source/GameplayGraph/Private/Graph/Graph.cpp
@@ -818,21 +818,82 @@ void operator>>(const IGraphDeserialization& Input, UGraph& Graph)
 					return FGraphIslandHandle{};
 				}
 
-				TArray<FGraphVertexHandle> IslandVertices;
-				IslandVertices.Reserve(Data.Vertices.Num());
+				TArray<FGraphVertexHandle> NewIslandVertices;
+				NewIslandVertices.Reserve(Data.Vertices.Num());
+
+				// This is a fallback to handle potentially degenerate cases where serialization
+				// has said that a vertex is in more than one island. In that case, we treat each a
+				// vertex as having an "implicit edge" to itself. This means that all these islands
+				// should be merged instead of creating a new one.
+				TArray<FGraphIslandHandle> ExistingIslands;
+				UGraphIsland* TargetMergeIsland = nullptr;
 		
 				for (const FGraphVertexHandle& VertexHandle : Data.Vertices)
 				{
 					FGraphVertexHandle CompleteHandle = Graph.GetCompleteNodeHandle(VertexHandle);
-					if (CompleteHandle.IsComplete())
+					if (UGraphVertex* Vertex = CompleteHandle.GetVertex())
 					{
-						IslandVertices.Add(CompleteHandle);
+						if (UGraphIsland* ParentIsland = Vertex->GetParentIsland().GetIsland())
+						{
+							ExistingIslands.Add(ParentIsland->Handle());
+
+							if (TargetMergeIsland)
+							{
+								if (ParentIsland->Num() > TargetMergeIsland->Num())
+								{
+									TargetMergeIsland = ParentIsland;
+								}
+							}
+							else
+							{
+								TargetMergeIsland = ParentIsland;
+							}
+						}
+						else
+						{
+							NewIslandVertices.Add(CompleteHandle);
+						}
 					}
 				}
 
-				FGraphIslandHandle NewIslandHandle = Graph.CreateIsland(IslandVertices, InHandle.GetUniqueIndex());
-				AllIslands.Add(NewIslandHandle);
-				return NewIslandHandle;
+				if (ExistingIslands.IsEmpty() || !TargetMergeIsland)
+				{
+					FGraphIslandHandle NewIslandHandle = Graph.CreateIsland(NewIslandVertices, InHandle.GetUniqueIndex());
+					AllIslands.Add(NewIslandHandle);
+					return NewIslandHandle;
+				}
+				else
+				{
+					// Stick all new vertices into the target merge island. Also all the ExistingIslands will get merged into the TargetMergeIsland as well.
+					// This handles the case where this new island holds vertices that are different islands already...which would be crazy!
+					for (const FGraphVertexHandle& VertexHandle : NewIslandVertices)
+					{
+						TargetMergeIsland->AddVertex(VertexHandle);
+					}
+
+					for (const FGraphIslandHandle& IslandHandle : ExistingIslands)
+					{
+						if (IslandHandle == TargetMergeIsland->Handle())
+						{
+							continue;
+						}
+
+						if (UGraphIsland* Island = IslandHandle.GetIsland())
+						{
+							TSet<FGraphVertexHandle> IslandVertices = Island->GetVertices();
+							for (const FGraphVertexHandle& VertexHandle : IslandVertices)
+							{
+								TargetMergeIsland->AddVertex(VertexHandle);
+							}
+
+							ensure(Island->Num() == 0);
+							Graph.RemoveIsland(IslandHandle);
+						}
+					}
+
+					// Return a null handle since this is technically not a new island.
+					return FGraphIslandHandle{};
+				}
 			}
 		);
 	}
diff --git a/Engine/Plugins/Experimental/GameplayGraph/Source/GameplayGraph/Public/Graph/GraphDefaultSerialization.h b/Engine/Plugins/Experimental/GameplayGraph/Source/GameplayGraph/Public/Graph/GraphDefaultSerialization.h
--- a/Engine/Plugins/Experimental/GameplayGraph/Source/GameplayGraph/Public/Graph/GraphDefaultSerialization.h
+++ b/Engine/Plugins/Experimental/GameplayGraph/Source/GameplayGraph/Public/Graph/GraphDefaultSerialization.h
@@ -173,6 +173,11 @@ public:
 	{
 		for (const TPair<FGraphIslandHandle, FSerializedIslandData>& Serialized : Data.Islands)
 		{
+			if (Serialized.Value.Vertices.IsEmpty())
+			{
+				continue;
+			}
+
 			FIslandConstructionData Construction;
 			Construction.Vertices = Serialized.Value.Vertices;
 
diff --git a/Engine/Plugins/Experimental/GameplayGraph/Tests/Private/GraphSerializationTests.cpp b/Engine/Plugins/Experimental/GameplayGraph/Tests/Private/GraphSerializationTests.cpp
--- a/Engine/Plugins/Experimental/GameplayGraph/Tests/Private/GraphSerializationTests.cpp
+++ b/Engine/Plugins/Experimental/GameplayGraph/Tests/Private/GraphSerializationTests.cpp
@@ -4,7 +4,7 @@
 
 #include "Graph/GraphDefaultSerialization.h"
 
-TEST_CASE_METHOD(FTestGraphBuilder, "Graph::Serialization::Write", "[graph]")
+TEST_CASE_METHOD(FTestGraphBuilder, "Graph::Serialization::Write", "[graph][serialization]")
 {
 	PopulateVertices(6, true);
 	BuildLinearEdges(3);
@@ -73,7 +73,7 @@ TEST_CASE_METHOD(FTestGraphBuilder, "Graph::Serialization::Write", "[graph]")
 	}
 }
 
-TEST_CASE_METHOD(FTestGraphBuilder, "Graph::Serialization::Read (normal)", "[graph]")
+TEST_CASE_METHOD(FTestGraphBuilder, "Graph::Serialization::Read (normal)", "[graph][serialization]")
 {
 	FSerializableGraph SerializedGraph;
 	SerializedGraph.Properties.bGenerateIslands = true;
@@ -144,7 +144,7 @@ TEST_CASE_METHOD(FTestGraphBuilder, "Graph::Serialization::Read (normal)", "[gra
 	}
 }
 
-TEST_CASE_METHOD(FTestGraphBuilder, "Graph::Serialization::Read::Errors::Vertex::Invalid Index", "[graph]")
+TEST_CASE_METHOD(FTestGraphBuilder, "Graph::Serialization::Read::Errors::Vertex::Invalid Index", "[graph][serialization]")
 {
 	FSerializableGraph SerializedGraph;
 	SerializedGraph.Properties.bGenerateIslands = true;
@@ -157,7 +157,7 @@ TEST_CASE_METHOD(FTestGraphBuilder, "Graph::Serialization::Read::Errors::Vertex:
 	CHECK(Graph->NumIslands() == 0);
 }
 
-TEST_CASE_METHOD(FTestGraphBuilder, "Graph::Serialization::Read::Errors::Edge::Invalid Vertex Index", "[graph]")
+TEST_CASE_METHOD(FTestGraphBuilder, "Graph::Serialization::Read::Errors::Edge::Invalid Vertex Index", "[graph][serialization]")
 {
 	FSerializableGraph SerializedGraph;
 	SerializedGraph.Properties.bGenerateIslands = true;
@@ -172,7 +172,7 @@ TEST_CASE_METHOD(FTestGraphBuilder, "Graph::Serialization::Read::Errors::Edge::I
 	CHECK(Graph->NumIslands() == 0);
 }
 
-TEST_CASE_METHOD(FTestGraphBuilder, "Graph::Serialization::Read::Errors::Island::Invalid Index", "[graph]")
+TEST_CASE_METHOD(FTestGraphBuilder, "Graph::Serialization::Read::Errors::Island::Invalid Index", "[graph][serialization]")
 {
 	FSerializableGraph SerializedGraph;
 	SerializedGraph.Properties.bGenerateIslands = true;
@@ -189,7 +189,7 @@ TEST_CASE_METHOD(FTestGraphBuilder, "Graph::Serialization::Read::Errors::Island:
 	CHECK(Graph->GetIslands().Contains(FGraphIslandHandle { FGraphUniqueIndex { FGuid { 0, 0, 0, 0 } }, nullptr }) == false);
 }
 
-TEST_CASE_METHOD(FTestGraphBuilder, "Graph::Serialization::Read::Errors::Island::Invalid Vertex Index", "[graph]")
+TEST_CASE_METHOD(FTestGraphBuilder, "Graph::Serialization::Read::Errors::Island::Invalid Vertex Index", "[graph][serialization]")
 {
 	FSerializableGraph SerializedGraph;
 	SerializedGraph.Properties.bGenerateIslands = true;
@@ -212,7 +212,7 @@ TEST_CASE_METHOD(FTestGraphBuilder, "Graph::Serialization::Read::Errors::Island:
 	CHECK(LoadedIsland->GetVertices().Contains(FGraphVertexHandle { FGraphUniqueIndex { FGuid { 0, 0, 0, 3 } }, nullptr }) == false);
 }
 
-TEST_CASE_METHOD(FTestGraphBuilder, "Graph::Serialization::Read::Errors::Island::Edge Island Mismatch::Merge", "[graph]")
+TEST_CASE_METHOD(FTestGraphBuilder, "Graph::Serialization::Read::Errors::Island::Edge Island Mismatch::Merge", "[graph][serialization]")
 {
 	FSerializableGraph SerializedGraph;
 	SerializedGraph.Properties.bGenerateIslands = true;
@@ -243,7 +243,7 @@ TEST_CASE_METHOD(FTestGraphBuilder, "Graph::Serialization::Read::Errors::Island:
 	CHECK(LoadedIsland->GetVertices().Contains(SerializedGraph.Vertices[2]) == true);
 }
 
-TEST_CASE_METHOD(FTestGraphBuilder, "Graph::Serialization::Read::Errors::Island::Edge Island Mismatch::Split", "[graph]")
+TEST_CASE_METHOD(FTestGraphBuilder, "Graph::Serialization::Read::Errors::Island::Edge Island Mismatch::Split", "[graph][serialization]")
 {
 	FSerializableGraph SerializedGraph;
 	SerializedGraph.Properties.bGenerateIslands = true;
@@ -276,4 +276,105 @@ TEST_CASE_METHOD(FTestGraphBuilder, "Graph::Serialization::Read::Errors::Island:
 	CHECK(SplitIsland->GetVertices().Contains(SerializedGraph.Vertices[0]) == false);
 	CHECK(SplitIsland->GetVertices().Contains(SerializedGraph.Vertices[1]) == false);
 	CHECK(SplitIsland->GetVertices().Contains(SerializedGraph.Vertices[2]) == true);
+}
+
+TEST_CASE_METHOD(FTestGraphBuilder, "Graph::Serialization::Read::Errors::Island::Duplicate Vertex Assignment::No Edges", "[graph][serialization]")
+{
+	FSerializableGraph SerializedGraph;
+	SerializedGraph.Properties.bGenerateIslands = true;
+	SerializedGraph.Vertices.Add(FGraphVertexHandle { FGraphUniqueIndex::CreateUniqueIndex(), nullptr });
+	SerializedGraph.Vertices.Add(FGraphVertexHandle { FGraphUniqueIndex::CreateUniqueIndex(), nullptr });
+	SerializedGraph.Vertices.Add(FGraphVertexHandle { FGraphUniqueIndex::CreateUniqueIndex(), nullptr });
+
+	SerializedGraph.Islands.Add(FGraphIslandHandle { FGraphUniqueIndex{FGuid{1, 0, 0, 0}}, nullptr }, FSerializedIslandData { TArray { SerializedGraph.Vertices[0], SerializedGraph.Vertices[2] } });
+	SerializedGraph.Islands.Add(FGraphIslandHandle { FGraphUniqueIndex{FGuid{2, 0, 0, 0}}, nullptr }, FSerializedIslandData { TArray { SerializedGraph.Vertices[1] } });
+	SerializedGraph.Islands.Add(FGraphIslandHandle { FGraphUniqueIndex{FGuid{3, 0, 0, 0}}, nullptr }, FSerializedIslandData { TArray { SerializedGraph.Vertices[0], SerializedGraph.Vertices[1] } });
+
+	FDefaultGraphDeserialization Deserializer { SerializedGraph };
+	Deserializer >> *Graph;
+	CHECK(SerializedGraph.Properties == Graph->GetProperties());
+	CHECK(Graph->NumVertices() == 3);
+	CHECK(Graph->NumIslands() == 3);
+
+	CHECK(Graph->GetIslands().Contains(FGraphIslandHandle { FGraphUniqueIndex { FGuid { 1, 0, 0, 0 } }, nullptr }) == true);
+	CHECK(Graph->GetIslands().Contains(FGraphIslandHandle { FGraphUniqueIndex { FGuid { 2, 0, 0, 0 } }, nullptr }) == false);
+	CHECK(Graph->GetIslands().Contains(FGraphIslandHandle { FGraphUniqueIndex { FGuid { 3, 0, 0, 0 } }, nullptr }) == false);
+
+	UGraphIsland* Island1 = Graph->GetCompleteNodeHandle(SerializedGraph.Vertices[0]).GetVertex()->GetParentIsland().GetIsland();
+	REQUIRE(Island1 != nullptr);
+	CHECK(Island1->Num() == 1);
+	CHECK(Island1->GetVertices().Contains(SerializedGraph.Vertices[0]) == true);
+	CHECK(Island1->GetVertices().Contains(SerializedGraph.Vertices[1]) == false);
+	CHECK(Island1->GetVertices().Contains(SerializedGraph.Vertices[2]) == false);
+
+	UGraphIsland* Island2 = Graph->GetCompleteNodeHandle(SerializedGraph.Vertices[1]).GetVertex()->GetParentIsland().GetIsland();
+	REQUIRE(Island2 != nullptr);
+	CHECK(Island2->Num() == 1);
+	CHECK(Island2->GetVertices().Contains(SerializedGraph.Vertices[1]) == true);
+	CHECK(Island2->GetVertices().Contains(SerializedGraph.Vertices[0]) == false);
+	CHECK(Island2->GetVertices().Contains(SerializedGraph.Vertices[2]) == false);
+
+	UGraphIsland* Island3 = Graph->GetCompleteNodeHandle(SerializedGraph.Vertices[2]).GetVertex()->GetParentIsland().GetIsland();
+	REQUIRE(Island3 != nullptr);
+	CHECK(Island3->Num() == 1);
+	CHECK(Island3->GetVertices().Contains(SerializedGraph.Vertices[2]) == true);
+	CHECK(Island3->GetVertices().Contains(SerializedGraph.Vertices[0]) == false);
+	CHECK(Island3->GetVertices().Contains(SerializedGraph.Vertices[1]) == false);
+}
+
+TEST_CASE_METHOD(FTestGraphBuilder, "Graph::Serialization::Read::Errors::Island::Duplicate Vertex Assignment::With Edge", "[graph][serialization]")
+{
+	FSerializableGraph SerializedGraph;
+	SerializedGraph.Properties.bGenerateIslands = true;
+	SerializedGraph.Vertices.Add(FGraphVertexHandle { FGraphUniqueIndex::CreateUniqueIndex(), nullptr });
+	SerializedGraph.Vertices.Add(FGraphVertexHandle { FGraphUniqueIndex::CreateUniqueIndex(), nullptr });
+
+	SerializedGraph.Edges.Add(FSerializedEdgeData{ SerializedGraph.Vertices[0], SerializedGraph.Vertices[1] });
+
+	SerializedGraph.Islands.Add(FGraphIslandHandle { FGraphUniqueIndex{FGuid{1, 0, 0, 0}}, nullptr }, FSerializedIslandData { TArray { SerializedGraph.Vertices[0] } });
+	SerializedGraph.Islands.Add(FGraphIslandHandle { FGraphUniqueIndex{FGuid{2, 0, 0, 0}}, nullptr }, FSerializedIslandData { TArray { SerializedGraph.Vertices[1] } });
+	SerializedGraph.Islands.Add(FGraphIslandHandle { FGraphUniqueIndex{FGuid{3, 0, 0, 0}}, nullptr }, FSerializedIslandData { TArray { SerializedGraph.Vertices[0], SerializedGraph.Vertices[1] } });
+
+	FDefaultGraphDeserialization Deserializer { SerializedGraph };
+	Deserializer >> *Graph;
+	CHECK(SerializedGraph.Properties == Graph->GetProperties());
+	CHECK(Graph->NumVertices() == 2);
+	CHECK(Graph->NumIslands() == 1);
+
+	CHECK(Graph->GetIslands().Contains(FGraphIslandHandle { FGraphUniqueIndex { FGuid { 1, 0, 0, 0 } }, nullptr }) == true);
+	CHECK(Graph->GetIslands().Contains(FGraphIslandHandle { FGraphUniqueIndex { FGuid { 2, 0, 0, 0 } }, nullptr }) == false);
+	CHECK(Graph->GetIslands().Contains(FGraphIslandHandle { FGraphUniqueIndex { FGuid { 3, 0, 0, 0 } }, nullptr }) == false);
+
+	UGraphIsland* Island1 = Graph->GetCompleteNodeHandle(SerializedGraph.Vertices[0]).GetVertex()->GetParentIsland().GetIsland();
+	REQUIRE(Island1 != nullptr);
+	CHECK(Island1->Handle() == FGraphIslandHandle { FGraphUniqueIndex { FGuid { 1, 0, 0, 0 } }, nullptr });
+	CHECK(Island1->Num() == 2);
+	CHECK(Island1->GetVertices().Contains(SerializedGraph.Vertices[0]) == true);
+	CHECK(Island1->GetVertices().Contains(SerializedGraph.Vertices[1]) == true);
+}
+
+TEST_CASE_METHOD(FTestGraphBuilder, "Graph::Serialization::Read::Errors::Island::Zero Size", "[graph][serialization]")
+{
+	FSerializableGraph SerializedGraph;
+	SerializedGraph.Properties.bGenerateIslands = true;
+	SerializedGraph.Vertices.Add(FGraphVertexHandle { FGraphUniqueIndex::CreateUniqueIndex(), nullptr });
+	SerializedGraph.Vertices.Add(FGraphVertexHandle { FGraphUniqueIndex::CreateUniqueIndex(), nullptr });
+	SerializedGraph.Edges.Add(FSerializedEdgeData{ SerializedGraph.Vertices[0], SerializedGraph.Vertices[1] });
+	SerializedGraph.Islands.Add(FGraphIslandHandle { FGraphUniqueIndex{FGuid{1, 0, 0, 0}}, nullptr }, FSerializedIslandData { TArray { SerializedGraph.Vertices[0], SerializedGraph.Vertices[1] } });
+	SerializedGraph.Islands.Add(FGraphIslandHandle { FGraphUniqueIndex{FGuid{2, 0, 0, 0}}, nullptr }, FSerializedIslandData { TArray<FGraphVertexHandle>{} });
+
+	FDefaultGraphDeserialization Deserializer { SerializedGraph };
+	Deserializer >> *Graph;
+	CHECK(SerializedGraph.Properties == Graph->GetProperties());
+	CHECK(Graph->NumVertices() == 2);
+	CHECK(Graph->NumIslands() == 1);
+
+	CHECK(Graph->GetIslands().Contains(FGraphIslandHandle { FGraphUniqueIndex { FGuid { 1, 0, 0, 0 } }, nullptr }) == true);
+	CHECK(Graph->GetIslands().Contains(FGraphIslandHandle { FGraphUniqueIndex { FGuid { 2, 0, 0, 0 } }, nullptr }) == false);
+
+	UGraphIsland* LoadedIsland = Graph->GetIslands().FindRef(FGraphIslandHandle {FGraphUniqueIndex{FGuid{1, 0, 0, 0}}, nullptr });
+	REQUIRE(LoadedIsland != nullptr);
+	CHECK(LoadedIsland->Num() == 2);
+	CHECK(LoadedIsland->GetVertices().Contains(SerializedGraph.Vertices[0]) == true);
+	CHECK(LoadedIsland->GetVertices().Contains(SerializedGraph.Vertices[1]) == true);
 }
\ No newline at end of file
