Mover: first edition of the readme doc

#rnx

diff --git a/Engine/Plugins/Experimental/Mover/README.md b/Engine/Plugins/Experimental/Mover/README.md
--- a/Engine/Plugins/Experimental/Mover/README.md
+++ b/Engine/Plugins/Experimental/Mover/README.md
@@ -2,10 +2,161 @@
 
 Mover is an Unreal Engine plugin to support movement of actors with rollback networking, using the Network Prediction Plugin. This plugin is the potential successor to Character Movement Component. The goal is to allow gameplay developers to focus on crafting motion without having to be experts in networking.
 
-The Mover plugin is Experimental. Many features are incomplete or missing. APIs and data formats are subject to change at any time.
+**The Mover plugin is Experimental. Many features are incomplete or missing. APIs and data formats are subject to change at any time.**
+
 
 ## Getting Started
 
+Add Mover and MoverExamples plugins to any existing project, or from a new Blank template project.
+
+Open the L_CharacterMovementBasics map, and activate Play-in-Editor (PIE).
+
+Recommended project settings to start:
+- Network Prediction / Preferred Ticking Policy: Fixed
+- Network Prediction / Simulated Proxy Network LOD: Interpolated
+- Engine General Settings / Use Fixed Framerate: Enabled
+
+If you have a fixed tick rate for Network Prediction, but a different tick rate under "Engine - General Settings", you may feel like the movement simulation is less smooth compared to camera movement or animation.
+
+
 ## Examples
 
-## Known Issues
\ No newline at end of file
+The MoverExamples plugin provides a collection of maps and actor examples, using a mixture of Blueprint and C++.
+
+Maps of interest:
+
+- **L_CharacterMovementBasics** has a variety of terrain features and movement examples.
+- **L_LayeredMoves** is focused on demonstrating many layered move types, with varying options.
+- **L_PhysicallyBasedCharacter** has an example of a physics-based pawn. See section below for more details.
+
+Pawn/Actor Blueprints of interest:
+
+- **AnimatedMannyPawn** is the simplest character and is based on the UE5 mannequin character used in the engine's template projects.
+- **AnimatedMannyPawnExtended** adds a variety of movement capabilities, such as dashing, vaulting, and ziplining.
+- **ScriptedAIMannyPawn** shows one way a non-player character can be controlled by Blueprint scripting.
+- **BP_SplineFollowerPlatform** shows a simple moving platform that predictably follows a defined path.
+
+
+## Concepts
+
+### Movement Modes
+
+A **movement mode** is an object that governs how your actor moves through space. It can both generate proposed movement and execute it. For example, walking, falling, climbing, etc. These modes look at inputs, decide how the character should move, and then attempt to actually move it through the world. 
+
+There is always one and only one mode active a time.
+
+### Layered Moves
+
+**Layered moves** represent temporary additional movement. For example, a constant force launching up into the air, a homing force that moves you towards an enemy, or even animation-driven root motion.
+
+They only generate a proposed move, and rely on the active movement mode to execute it. 
+
+Multiple layered moves can be active at a time, and their moves can be mixed with other influences. Their lifetime can have a duration or be instantaneous, as well as surviving transitions between modes.
+
+### Transitions
+
+**Transitions** are objects that evaluate whether a Mover actor should change its mode, based on the current state. It can also trigger side effects when activated.
+
+They can be associated with a particular mode, only evaluated when that mode is active. Or they can be global and evalutated no matter what mode the Mover actor is in.
+
+The use of transitions is optional, with other methods of switching modes available.
+
+
+### Other Concepts
+
+- **Composable Input and Sync State:** Inputs are authored by the controlling owner, and influence a movement simulation step. Sync state is a snapshot that describes a Mover actor's movement at a point in time. Both inputs and sync state can have custom struct data attached to them dynamically. 
+
+- **Shared Settings:** Collections of properties that multiple movement objects share, to avoid duplication of settings between decoupled objects.  The list is managed by MoverComponent based on which settings classes its modes call for.
+
+- **Movement Utility Libraries:** (optional) These are collection of functions useful for crafting movement, typically callable from Blueprints. When implementing the default movement set, we attempted to break the methods into these libraries where possible, so that developers can make use of them in their own movement.
+
+- **Sim Blackboard:** (optional) This is a way for decoupled systems to share information or cache computations between simulation ticks without adding to the official simulation state. Note that this system is not yet rollback-friendly. 
+
+- **Move Record:** (optional) This is a mechanism for tracking combinations of moves and marking which ones should affect the final state of the moving actor.  For example, a movement that results in a character stepping up onto a short ledge consists of both horizontal and vertical movement, but we don't want the vertical movement to be included when computing the actor's post-move velocity. The record(s) for vertical movement would be marked as non-contributing. 
+
+
+## Comparing with Character Movement Component (CMC)
+
+- **Movement Modes** are very similar in both systems, but modular in Mover
+- **Layered Moves** are similar to Root Motion Sources (RMS)
+- **Transitions** are a new concept
+- Movement from modes and layered moves can be mixed together
+- It is easier to add custom movement modes, even from plugins and at runtime
+- The DefaultMovementSet in Mover is similar to the modes built in to CharacterMovementComponent (Walking, Falling, Flying, etc.)
+
+The Mover plugin provides a DefaultMovementSet that is similar to the modes that CharacterMovementComponent provides. This movement set assumes a similar actor composition, with a Capsule primitive component as the root, with a skeletal mesh attached to it.
+
+MoverComponent does not require your actor class to derive from ACharacter.
+
+MoverComponent requires a root SceneComponent, but it does not have to be a singular vertically-oriented capsule or even a PrimitiveComponent. Developers are free to create Mover actors with no collision primitives if they wish.
+
+MoverComponent does not require or assume a skeletal mesh as the visual representation.
+
+In CMC, adding custom data to be passed between clients and server required subclassing the component and overriding key functions. The Mover plugin allows custom input and state data to be added dynamically at runtime, without customizing the MoverComponent.
+
+Networking Model: 
+- In CMC, owning clients send a combination of inputs and state as a "move" at the client's frame rate. The server receives them and performs the same move immediately, then compares state to decide if a correction is needed and replies with either an acknowledgement of the move or a block of corrective state data. 
+- In Mover / Network Prediction, all clients and server are attempting to simulate on a shared timeline, with clients running predictively ahead of the server by a small amount of time. Clients author inputs for a specific simulation time/frame, and that is all they send to the server. The server buffers these inputs until their simulation time comes. After performing all movement updates, the server broadcasts state to all clients and the clients decide whether a correction (rollback + resim) is needed.
+
+Unlike CMC, the state of the Mover actor is not directly modifiable externally at any time. For example, there is no Velocity property to directly manipulate. Instead, developers must make use of modes and layered moves to affect change during the next available simulation tick.  Additionally, player-provided inputs such as move input and button presses must be combined into a single Input Command for the movement simulation tick, rather than immediately affecting the Mover actor's state.  Depending on your project settings, you may have player input from several frames contributing to a single movement simulation tick.
+
+
+## Debugging
+
+### Gameplay Debugger Tool
+
+Visualization and state readout information is available through the Gameplay Debugger Tool (GDT). To activate the gameplay debugger tool, typically via the **'** key, and toggle the Mover category using the NumPad. The locally-controlled player character will be selected by default, but you can change this via GDT input and **gdt.\*** console commands.
+
+### Logging
+
+Output log information coming from the Mover plugin will have the **LogMover** category.
+
+### Console Commands
+
+There are a variety of useful console commands with this prefix: **Mover.\***
+
+
+## Physics-driven Character Example
+
+Included in MoverExamples is a physics-driven version of the Manny pawn.  It is an experiment within this experimental plugin, and should be treated as such.  To try it out, open the L_PhysicallyBasedCharacter map in the MoverExamples plugin. Make sure you adjust your project's settings according to the text inside the map.
+
+Other MoverExamples pawns, and the CharacterMovementComponent, use a "kinematic" movement style where the pawn's shape is moved by testing the surroundings in an ad hoc manner. Responses to physics forces, and interactions with physics-simulated objects, are difficult to implement. With physics driving the movement, the pawn can realistically apply forces to other objects and have them applied right back.
+
+This physics-based character is NOT using the Network Prediction plugin. Instead, it is driven by the Chaos Networked Physics system, which has many similarities in how it operates. Under the hood, a Character Ground Constraint from the physics system handles performing the proposed motion that the Mover system generates. The physics simulation is running at a fixed tick rate, ahead of the game thread. As the simulation progresses, the game thread's character representation interpolates towards the most recent physics state.
+
+Notes:
+- Due to the async nature of the physics simulation, there is some additional lag time between player input and affecting the pawn's movement seen on screen.
+- Other physics objects that the character interacts with should be set up for replication. Otherwise this will lead to differences between client and server.
+- Interactions with moving non-physics objects will likely show poor results, due to ticking differences between the physics simulation and the rest of the game world.
+- Various gameplay events may not be connected or are unreliable. 
+- Crafting customized movement may be less flexible without also using a modified physics constraint and solver.
+- Physics-driven Mover actors are not synchronized with those using Network Prediction 
+
+
+## FAQ
+
+- **Should my project switch to Mover from CharacterMovementComponent?**  This depends greatly on the scope of your project and will require some due diligence. The Mover plugin is newly experimental and hasn't gone through the rigors of a shipped project yet. There are many gaps in functionality, and little time has been spent on scaling/performance. Single-player or games with smaller character counts will be more feasible. 
+
+- **Does Mover fix synchronization issues between movement and the Gameplay Ability System?** The short answer is no. GAS still has its own independent replication methods. The use of Network Prediction opens the door for GAS (or other systems) to integrate with Network Prediction and achieve good synchronization with movement.
+
+- **What about single-player games?** Mover is useful for single-player games as well, you just won't be making use of the networking and rollback features. Since there's no need to synchronize simulation time with a server, consider changing the project setting "Network Prediction / Preferred Ticking Policy" to "Independent" mode. This will make the movement simulation tick at the same rate as the game world. 
+
+
+## Limitations and Known Issues
+
+**See also: Network Prediction Plugin's documentation** Some of the Mover plugin's current limitations come from its dependency on the Network Prediction plugin. Please review its documentation for more info.
+
+**Limited Blueprinting Support:** Blueprint functionality isn't 100% supported yet. There are certain things that still require native C++ code, or are clunky to implement in Blueprints.
+
+**Arbitrary gravity, collision shapes, etc.:** Although the core MoverComponent tries to make as few requirements as possible on the composition of the actor, the default movement set has more rigid assumptions.  For example, the default movement set currently assumes a capsule shape. Additionally, some features such as arbitrary gravity are not fully supported in all cases yet.
+
+**Animation of Sim Proxy Example Characters:** Animation on another client's pawn (a sim proxy) may not be fully replicated during certain actions. This will be improved in a future release.
+
+**Forward-Predicted Sim Proxy Characters:**  Characters controlled by other players are typically poor candidates for forward prediction, which relies on past inputs to predict future movement. Acceleration and direction changes, as well as action inputs like jumping, are unpredictable and will be the source of frequent mispredictions that can give the sim proxy character popping or choppy movement. Consider using Interpolated mode for the "Simulated Proxy Network LOD" project setting, which will give smooth results at the cost of some visual latency.
+
+**Cooked Data Optimization Can Lead to Missing Data:** Cooked builds may have some MoverComponent data missing, resulting in a non-functional actor. Disable the actor's "Generate Optimized Blueprint Component Data" option (bOptimizeBPComponentData in C++). 
+
+**MoverComponent Defaults Are Currently Focused on Characters:** The current component has more API and data than would be required for Mover actors that are not Character-like. Our goal is to make a minimalist base component with slimmer API and property set, with a specialization for Character-like actors implemented separately as part of the default movement set. 
+
+**Sim Blackboard Is Not Yet Rollback-Friendly:** When rollbacks occur, the blackboard's contents are simply invalidated. Although the blackboard is useful for avoiding repeating computations from a prior frame, movement logic cannot always rely on it having a valid entry and should have a fallback.
+
