Mover readme updates, focused on current known issues

#rnx

diff --git a/Engine/Plugins/Experimental/Mover/README.md b/Engine/Plugins/Experimental/Mover/README.md
--- a/Engine/Plugins/Experimental/Mover/README.md
+++ b/Engine/Plugins/Experimental/Mover/README.md
@@ -146,12 +146,22 @@ Notes:
 
 **See also: Network Prediction Plugin's documentation** Some of the Mover plugin's current limitations come from its dependency on the Network Prediction plugin. Please review its documentation for more info.
 
+**Ticking Ordering:** All Network Prediction simulations tick before Unreal's world tick groups and always in a consistent order, making it difficult to implement mechanics where tight ticking dependencies are needed.
+
+**Fixed Tick Simulation + Variable Engine/Rendering Rate:** When using Fixed Tick simulation with a variable Engine ticking (rendering) rate, you could go several rendered frames in between simulation ticks, or even have multiple simulation ticks occur during a single rendered frame. This has implications for things like input capturing, where you may need to combine several frames of sampled input into an input command for a single simulation tick.
+
+Additionally, when the simulation runs at a lower rate than rendering, the result is rough-looking movement that is unacceptable for most projects. Smoothing/interpolation is an active work area, and we anticipate having an initial solution soon.
+
+**Setting a Fixed Tick Rate:** There is a bug with Network Prediction's "Fixed Tick Frame Rate" setting not being respected. If you choose "Fixed Tick" mode, you need to set your project's "Engine/General Settings/Fixed Frame Rate" setting, even if not using it. Note that you may have to temporarily enable "Use Fixed Frame Rate" in order to edit the value.  Due to this, it is not currently possible to run the simulation at one fixed rate with engine/rendering ticking at a different fixed rate, like 20 fps vs 60 fps.
+
 **Limited Blueprinting Support:** Blueprint functionality isn't 100% supported yet. There are certain things that still require native C++ code, or are clunky to implement in Blueprints.
 
-**Arbitrary gravity, collision shapes, etc.:** Although the core MoverComponent tries to make as few requirements as possible on the composition of the actor, the default movement set has more rigid assumptions.  For example, the default movement set currently assumes a capsule shape. Additionally, some features such as arbitrary gravity are not fully supported in all cases yet.
+**Arbitrary Gravity, Collision Shapes, etc.:** Although the core MoverComponent tries to make as few requirements as possible on the composition of the actor, the default movement set has more rigid assumptions.  For example, the default movement set currently assumes a capsule shape. Additionally, some features such as arbitrary gravity are not fully supported in all cases yet.
 
 **Animation of Sim Proxy Example Characters:** Animation on another client's pawn (a sim proxy) may not be fully replicated during certain actions. This will be improved in a future release.
 
+**Lack of Replay Support:** Replays do not currently work with Network Prediction simulations.
+
 **Forward-Predicted Sim Proxy Characters:**  Characters controlled by other players are typically poor candidates for forward prediction, which relies on past inputs to predict future movement. Acceleration and direction changes, as well as action inputs like jumping, are unpredictable and will be the source of frequent mispredictions that can give the sim proxy character popping or choppy movement. Consider using Interpolated mode for the "Simulated Proxy Network LOD" project setting, which will give smooth results at the cost of some visual latency.
 
 **Cooked Data Optimization Can Lead to Missing Data:** Cooked builds may have some MoverComponent data missing, resulting in a non-functional actor. Disable the actor's "Generate Optimized Blueprint Component Data" option (bOptimizeBPComponentData in C++). 
