Cleaning warnings in StateGraphTests

diff --git a/Engine/Plugins/Experimental/StateGraph/Tests/StateGraphTests/Private/StateGraphTestsExample.cpp b/Engine/Plugins/Experimental/StateGraph/Tests/StateGraphTests/Private/StateGraphTestsExample.cpp
--- a/Engine/Plugins/Experimental/StateGraph/Tests/StateGraphTests/Private/StateGraphTestsExample.cpp
+++ b/Engine/Plugins/Experimental/StateGraph/Tests/StateGraphTests/Private/StateGraphTestsExample.cpp
@@ -214,6 +214,9 @@ protected:
 
 TEST_CASE("Example demonstrating the features of StateGraph", "[Example]")
 {
+	// Ignore potenital depercated .ini warnings in this test.
+	LogConfig.SetVerbosity(ELogVerbosity::Error);
+
 	// Engine does this, but it's needed here to ensure GConfig is initialized for the test.
 	FConfigCacheIni::InitializeConfigSystem();
 
diff --git a/Engine/Plugins/Experimental/StateGraph/Tests/StateGraphTests/Private/StateGraphTestsUnit.cpp b/Engine/Plugins/Experimental/StateGraph/Tests/StateGraphTests/Private/StateGraphTestsUnit.cpp
--- a/Engine/Plugins/Experimental/StateGraph/Tests/StateGraphTests/Private/StateGraphTestsUnit.cpp
+++ b/Engine/Plugins/Experimental/StateGraph/Tests/StateGraphTests/Private/StateGraphTestsUnit.cpp
@@ -7,6 +7,26 @@
 
 namespace UE::StateGraphTests::Unit
 {
+	
+class StateGraphDisableWarningsLog
+{
+public:
+	StateGraphDisableWarningsLog()
+		: OldVerbosity(LogStateGraph.GetVerbosity())
+	{
+		LogStateGraph.SetVerbosity(ELogVerbosity::Error);
+	}
+
+	~StateGraphDisableWarningsLog()
+	{
+		if (OldVerbosity != LogStateGraph.GetVerbosity())
+		{
+			LogStateGraph.SetVerbosity(OldVerbosity);
+		}
+	}
+
+	ELogVerbosity::Type OldVerbosity;
+};
 
 TEST_CASE("FStateGraph Basic Tests", "[FStateGraph]")
 {
@@ -134,6 +154,8 @@ TEST_CASE("CreateNode for each supported type", "[CreateNodeFunctions]")
 
 TEST_CASE("Node dependencies", "[NodeDependencies]")
 {
+	StateGraphDisableWarningsLog ScopeDisable;
+
 	FStateGraphRef StateGraph(MakeShared<FStateGraph>("Test"));
 
 	FStateGraphNodeRef TestNodeA = StateGraph->CreateNode<FTestNode>("TestNodeA");
@@ -175,6 +197,7 @@ TEST_CASE("Node dependencies", "[NodeDependencies]")
 
 TEST_CASE("Blocked graph", "[BlockedGraph]")
 {
+	StateGraphDisableWarningsLog ScopeDisable;
 	FStateGraphRef StateGraph(MakeShared<FStateGraph>("Test"));
 
 	FStateGraphNodeRef TestNodeA = StateGraph->CreateNode<FTestNode>("TestNodeA", TSet<FName>({ "TestNodeB" }));
@@ -211,9 +234,13 @@ TEST_CASE("Adding and reusing nodes", "[AddingNodes]")
 	CHECK(StateGraphB->GetStatus() == FStateGraph::EStatus::Completed);
 	CHECK(TestNodeA->GetStatus() == FStateGraphNode::EStatus::Completed);
 
-	StateGraphA->Reset();
-	CHECK(StateGraphA->GetStatus() == FStateGraph::EStatus::NotStarted);
-	CHECK(!StateGraphA->AddNode(TestNodeA));
+	{
+		StateGraphDisableWarningsLog ScopeDisable;
+
+		StateGraphA->Reset();
+		CHECK(StateGraphA->GetStatus() == FStateGraph::EStatus::NotStarted);
+		CHECK(!StateGraphA->AddNode(TestNodeA));
+	}
 
 	StateGraphB->RemoveNode(TestNodeA->GetName());
 	CHECK(StateGraphA->AddNode(TestNodeA));
@@ -221,8 +248,11 @@ TEST_CASE("Adding and reusing nodes", "[AddingNodes]")
 	CHECK(StateGraphA->GetStatus() == FStateGraph::EStatus::Completed);
 	CHECK(TestNodeA->GetStatus() == FStateGraphNode::EStatus::Completed);
 
-	FStateGraphNodeRef TestNodeA2 = MakeShared<FTestNode>("TestNodeA");
-	CHECK(!StateGraphA->AddNode(TestNodeA2));
+	{
+		StateGraphDisableWarningsLog ScopeDisable;
+		FStateGraphNodeRef TestNodeA2 = MakeShared<FTestNode>("TestNodeA");
+		CHECK(!StateGraphA->AddNode(TestNodeA2));
+	}
 }
 
 TEST_CASE("Removing nodes", "[RemoveNodes]")
@@ -417,6 +447,7 @@ TEST_CASE("Node timeout", "[NodeTimeout]")
 
 	while (StateGraph->GetStatus() == FStateGraph::EStatus::Waiting)
 	{
+		StateGraphDisableWarningsLog ScopeDisable;
 		FTSTicker::GetCoreTicker().Tick(0.1f);
 	}
 
@@ -431,6 +462,8 @@ TEST_CASE("Node timeout", "[NodeTimeout]")
 
 TEST_CASE("State graph config", "[StateGraphConfig]")
 {
+	StateGraphDisableWarningsLog ScopeDisable;
+	LogConfig.SetVerbosity(ELogVerbosity::Error);
 	FConfigCacheIni::InitializeConfigSystem();
 	FConfigFile* EngineIni = GConfig->FindConfigFile(GEngineIni);
 	check(EngineIni);
@@ -459,6 +492,8 @@ TEST_CASE("State graph config", "[StateGraphConfig]")
 
 TEST_CASE("Node config", "[NodeConfig]")
 {
+	StateGraphDisableWarningsLog ScopeDisable;
+	LogConfig.SetVerbosity(ELogVerbosity::Error);
 	FConfigCacheIni::InitializeConfigSystem();
 	FConfigFile* EngineIni = GConfig->FindConfigFile(GEngineIni);
 	check(EngineIni);
