Fix for a GameplayGraph Test that was improperly initializing a struct.

The constructor does not allow brace initialization for structs if the struct's variables are default initialized in the header with its definition. FGraphProperties::bGenerateIslands is being initialized to true as part of the struct definition. Rather than change the struct definition, the brace initialization in this test was just removed since it was re-setting bGenerateIslands to its default value

diff --git a/Engine/Plugins/Experimental/GameplayGraph/Tests/Private/TestGraphBuilder.cpp b/Engine/Plugins/Experimental/GameplayGraph/Tests/Private/TestGraphBuilder.cpp
--- a/Engine/Plugins/Experimental/GameplayGraph/Tests/Private/TestGraphBuilder.cpp
+++ b/Engine/Plugins/Experimental/GameplayGraph/Tests/Private/TestGraphBuilder.cpp
@@ -6,7 +6,7 @@
 FTestGraphBuilder::FTestGraphBuilder()
 {
 	Graph = NewObject<UGraph>();
-	Graph->InitializeFromProperties(FGraphProperties{ true });
+	Graph->InitializeFromProperties(FGraphProperties());
 }
 
 FTestGraphBuilder::~FTestGraphBuilder()
