Mover: documentation updates

#virtualized

diff --git a/Engine/Plugins/Experimental/Mover/Mover.uplugin b/Engine/Plugins/Experimental/Mover/Mover.uplugin
--- a/Engine/Plugins/Experimental/Mover/Mover.uplugin
+++ b/Engine/Plugins/Experimental/Mover/Mover.uplugin
@@ -3,7 +3,7 @@
 	"Version": 1,
 	"VersionName": "1.0",
 	"FriendlyName": "Mover",
-	"Description": "Mover is an Unreal Engine plugin to support movement of actors with rollback networking.",
+	"Description": "Mover is an Unreal Engine plugin to support movement of actors with rollback networking.\nPlease refer to the README document for information about getting started, an overview of concepts, and known issues.",
 	"Category": "Gameplay",
 	"CreatedBy": "Epic Games, Inc.",
 	"CreatedByURL": "https://epicgames.com",
diff --git a/Engine/Plugins/Experimental/Mover/README.md b/Engine/Plugins/Experimental/Mover/README.md
--- a/Engine/Plugins/Experimental/Mover/README.md
+++ b/Engine/Plugins/Experimental/Mover/README.md
@@ -1,6 +1,6 @@
 # Mover Plugin
 
-Mover is an Unreal Engine plugin to support movement of actors with rollback networking, using the Network Prediction Plugin. This plugin is the potential successor to Character Movement Component. The goal is to allow gameplay developers to focus on crafting motion without having to be experts in networking.
+Mover is an Unreal Engine plugin to support movement of actors with rollback networking, using the Network Prediction Plugin or Chaos Networked Physics. This plugin is the potential successor to Character Movement Component. The goal is to allow gameplay developers to focus on crafting motion without having to be experts in networking.
 
 **The Mover plugin is Experimental. Many features are incomplete or missing. APIs and data formats are subject to change at any time.**
 
@@ -14,10 +14,9 @@ Open the L_CharacterMovementBasics map, and activate Play-in-Editor (PIE).
 Recommended project settings to start:
 - Network Prediction / Preferred Ticking Policy: Fixed
 - Network Prediction / Simulated Proxy Network LOD: Interpolated
-- Engine General Settings / Use Fixed Framerate: Enabled
-
-If you have a fixed tick rate for Network Prediction, but a different tick rate under "Engine - General Settings", you may feel like the movement simulation is less smooth compared to camera movement or animation.
+- Network Prediction / Enable Fixed Tick Smoothing: true
 
+This will have Mover objects simulating at a fixed rate, which is different than the variable rendering frame rate. The smoothing option provides visual compensation. Without it, you may feel like the movement simulation is rougher than the camera movement or skeletal animation.
 
 ## Examples
 
@@ -27,14 +26,15 @@ Maps of interest:
 
 - **L_CharacterMovementBasics** has a variety of terrain features and movement examples.
 - **L_LayeredMoves** is focused on demonstrating many layered move types, with varying options.
+- **L_PathfindingMovement** is focused on demonstrating Mover and different AI-driven movement with pathfinding including NavWalking mode.
 - **L_PhysicallyBasedCharacter** has an example of a physics-based pawn. See section below for more details.
 
 Pawn/Actor Blueprints of interest:
 
 - **AnimatedMannyPawn** is the simplest character and is based on the UE5 mannequin character used in the engine's template projects.
 - **AnimatedMannyPawnExtended** adds a variety of movement capabilities, such as dashing, vaulting, and ziplining.
-- **ScriptedAIMannyPawn** shows one way a non-player character can be controlled by Blueprint scripting.
-- **BP_SplineFollowerPlatform** shows a simple moving platform that predictably follows a defined path.
+- **PathFollowingMannyPawn** is set up to support AI-driven pathfollowing as an example of nav-mesh movement.
+- **BP_SplineFollowerPlatform** shows a simple moving platform that predictably follows a defined path. This is an Actor, but not a Pawn.
 
 
 ## Concepts
@@ -47,26 +47,52 @@ There is always one and only one mode active a time.
 
 ### Layered Moves
 
-**Layered moves** represent temporary additional movement. For example, a constant force launching up into the air, a homing force that moves you towards an enemy, or even animation-driven root motion.
+**Layered moves** represent temporary additional movement. For example, a constant force briefly launching up into the air, a homing force that moves you towards an enemy, or even animation-driven root motion.
 
 They only generate a proposed move, and rely on the active movement mode to execute it. 
 
 Multiple layered moves can be active at a time, and their moves can be mixed with other influences. Their lifetime can have a duration or be instantaneous, as well as surviving transitions between modes.
 
-### Transitions
+### Instant Movement Effects
+
+**Instant movement effects** can affect movement state directly on a Mover-based actor during simulation without consuming time. They are only applied during a simulation tick and then removed.
+
+Multiple instant effects can be queued and will be executed sequentially as soon as possible. These execution windows occur at the beginning and end of each simulation step, including between substeps (for example, if the mode changed from falling to walking halfway through the sim tick).
+
+This could be used for teleporting, applying instantaneous launch velocities, or forcing a movement mode change. 
 
-**Transitions** are objects that evaluate whether a Mover actor should change its mode, based on the current state. It can also trigger side effects when activated.
+### Movement Modifiers
+
+**Movement Modifiers** are used to apply changes that *indirectly influence* the movement simulation, without proposing or executing any movement, but still in sync with the simulation.
+
+Common uses would be for movement "stance" changes like crouching, stealth, etc. or other mechanics requiring a number of settings to change together.
+
+### Transitions
 
-They can be associated with a particular mode, only evaluated when that mode is active. Or they can be global and evalutated no matter what mode the Mover actor is in.
+**Transitions** are objects that evaluate whether a Mover actor should change its mode, based on the current state. It can also trigger side effects when activated. They can set up to only evaluate when a particular mode is active, or they can be global and evalutated regardless of the current mode.
 
 The use of transitions is optional, with other methods of switching modes available.
 
+### Backend Liaisons
+
+A Mover-based actor requires a MoverComponent (an ActorComponent type) to operate. This is the main object that other systems interface with. Unlike many ActorComponents that use TickComponent to drive updating, Mover relies on an external "backend" to drive it.
+
+The liaison acts as an intermediary between the MoverComponent and the backend, passing function calls or events when it's time to produce input, advance the simulation, etc.  Current backend liaisons can interface with the Network Prediction Plugin or Chaos' Networked Physics system.  There is also a simple "standalone" liaison that can be used for non-networked play with kinematic movement.
+
 
 ### Other Concepts
 
 - **Composable Input and Sync State:** Inputs are authored by the controlling owner, and influence a movement simulation step. Sync state is a snapshot that describes a Mover actor's movement at a point in time. Both inputs and sync state can have custom struct data attached to them dynamically. 
 
-- **Shared Settings:** Collections of properties that multiple movement objects share, to avoid duplication of settings between decoupled objects.  The list is managed by MoverComponent based on which settings classes its modes call for.
+- **Shared Settings:** Collections of properties that multiple movement objects share, to avoid duplication of settings between decoupled objects.  The list is managed by MoverComponent based on which settings classes its modes call for. Your custom modes are not required to use Shared Settings and may approach settings differently, for example having a self-contained set of properties or referencing an external data table.
+
+- **Persistent Sync State Data Types:** There is almost always state data that you want to carry over from the previous frame, especially core movement data such as position.  Adding types to this array allows carry-over of types to automatically happen. Any state data *not* present on this list will have to be added every frame.
+
+- **Input Production:** Mover does not deal directly with Unreal's input systems such as Enhanced Input. Projects instead should translate input events into InputCmd structure(s) that get submitted just before simulation frames, via the ProduceInput interface.  AMoverExamplesCharacter's use of OnProduceInput shows a C++ method. The AnimatedMannyPawnExtended BP in the MoverExamples plugin adds additional inputs via the "On Produce Input" function implementation.
+
+- **Trajectory Prediction:** Mover provides a method to get samples that predict where an actor will end up in the future via the GetPredictedTrajectory function. This is useful in support of advanced animation techniques like Motion Matching.  As of 5.5, this is based on the most recent state and inputs. Options and customization will be expanded in the future.
+
+- **Nav Mover Component:** (optional) This is a component meant to house navigation movement specific settings and functionality. It is designed to live alongside the MoverComponent attached to the same actor.
 
 - **Movement Utility Libraries:** (optional) These are collection of functions useful for crafting movement, typically callable from Blueprints. When implementing the default movement set, we attempted to break the methods into these libraries where possible, so that developers can make use of them in their own movement.
 
@@ -79,7 +105,9 @@ The use of transitions is optional, with other methods of switching modes availa
 
 - **Movement Modes** are very similar in both systems, but modular in Mover
 - **Layered Moves** are similar to Root Motion Sources (RMS)
-- **Transitions** are a new concept
+- **Instant Effects** are a new concept, and are used to implement CMC features like launching and teleporting
+- **Movement Modifiers** are a new concept, and are used to implement CMC features like crouching
+- **Transitions** are a new concept, as a modular way to drive mode-changing evaluation logic
 - Movement from modes and layered moves can be mixed together
 - It is easier to add custom movement modes, even from plugins and at runtime
 - The DefaultMovementSet in Mover is similar to the modes built in to CharacterMovementComponent (Walking, Falling, Flying, etc.)
@@ -130,16 +158,16 @@ Notes:
 - Interactions with moving non-physics objects will likely show poor results, due to ticking differences between the physics simulation and the rest of the game world.
 - Various gameplay events may not be connected or are unreliable. 
 - Crafting customized movement may be less flexible without also using a modified physics constraint and solver.
-- Physics-driven Mover actors are not synchronized with those using Network Prediction 
+- Physics-driven Mover actors are not well-synchronized with the movement of non-physics actors, such as those manually moved via gameplay scripting (or even Mover actors using NPP).
 
 
 ## FAQ
 
-- **Should my project switch to Mover from CharacterMovementComponent?**  This depends greatly on the scope of your project and will require some due diligence. The Mover plugin is newly experimental and hasn't gone through the rigors of a shipped project yet. There are many gaps in functionality, and little time has been spent on scaling/performance. Single-player or games with smaller character counts will be more feasible. 
+- **Should my project switch to Mover from CharacterMovementComponent?**  This depends greatly on the scope of your project and will require some due diligence. The Mover plugin is experimental and hasn't gone through the rigors of a shipped project yet. There are many gaps in functionality, and little time has been spent on scaling/performance. Single-player or games with smaller character counts will be more feasible. 
 
 - **Does Mover fix synchronization issues between movement and the Gameplay Ability System?** The short answer is no. GAS still has its own independent replication methods. The use of Network Prediction opens the door for GAS (or other systems) to integrate with Network Prediction and achieve good synchronization with movement.
 
-- **What about single-player games?** Mover is useful for single-player games as well, you just won't be making use of the networking and rollback features. Since there's no need to synchronize simulation time with a server, consider changing the project setting "Network Prediction / Preferred Ticking Policy" to "Independent" mode. This will make the movement simulation tick at the same rate as the game world. 
+- **What about single-player games?** Mover is useful for single-player games as well, you just won't be making use of the networking and rollback features. Consider changing the project setting "Network Prediction / Preferred Ticking Policy" to "Independent" mode. This will make the movement simulation tick at the same rate as the game world.  You can even eliminate all NPP overhead by choosing "MoverStandaloneLiaisonComponent" as your "Backend Class" property on your actors' MoverComponents.  Physics-based pawns should also work well for single-player projects.
 
 
 ## Limitations and Known Issues
@@ -150,18 +178,14 @@ Notes:
 
 **Fixed Tick Simulation + Variable Engine/Rendering Rate:** When using Fixed Tick simulation with a variable Engine ticking (rendering) rate, you could go several rendered frames in between simulation ticks, or even have multiple simulation ticks occur during a single rendered frame. This has implications for things like input capturing, where you may need to combine several frames of sampled input into an input command for a single simulation tick.
 
-Additionally, when the simulation runs at a lower rate than rendering, the result is rough-looking movement that is unacceptable for most projects. Smoothing/interpolation is an active work area, and we anticipate having an initial solution soon.
-
-**Setting a Fixed Tick Rate:** There is a bug with Network Prediction's "Fixed Tick Frame Rate" setting not being respected. If you choose "Fixed Tick" mode, you need to set your project's "Engine/General Settings/Fixed Frame Rate" setting, even if not using it. Note that you may have to temporarily enable "Use Fixed Frame Rate" in order to edit the value.  Due to this, it is not currently possible to run the simulation at one fixed rate with engine/rendering ticking at a different fixed rate, like 20 fps vs 60 fps.
+Additionally, when the simulation runs at a lower rate than rendering, the result is rough-looking movement that is unacceptable for most projects. The 5.5 release introduces options for visual smoothing when using Network Prediction. Set Project Settings' "Network Prediction / Enable Fixed Tick Smoothing" to true, and set your MoverComponent's "Smoothing Mode" to "Visual Component Offset".
 
 **Limited Blueprinting Support:** Blueprint functionality isn't 100% supported yet. There are certain things that still require native C++ code, or are clunky to implement in Blueprints.
 
-**Arbitrary Gravity, Collision Shapes, etc.:** Although the core MoverComponent tries to make as few requirements as possible on the composition of the actor, the default movement set has more rigid assumptions.  For example, the default movement set currently assumes a capsule shape. Additionally, some features such as arbitrary gravity are not fully supported in all cases yet.
+**Arbitrary Gravity, Collision Shapes, etc. in the Default Movement Set:** Although the core MoverComponent tries to make as few requirements as possible on the composition of the actor, the default movement set has more rigid assumptions.  For example, the default movement set currently assumes a capsule shape. Some features such as arbitrary gravity are not fully supported in all cases yet.
 
 **Animation of Sim Proxy Example Characters:** Animation on another client's pawn (a sim proxy) may not be fully replicated during certain actions. This will be improved in a future release.
 
-**Lack of Replay Support:** Replays do not currently work with Network Prediction simulations.
-
 **Forward-Predicted Sim Proxy Characters:**  Characters controlled by other players are typically poor candidates for forward prediction, which relies on past inputs to predict future movement. Acceleration and direction changes, as well as action inputs like jumping, are unpredictable and will be the source of frequent mispredictions that can give the sim proxy character popping or choppy movement. Consider using Interpolated mode for the "Simulated Proxy Network LOD" project setting, which will give smooth results at the cost of some visual latency.
 
 **Cooked Data Optimization Can Lead to Missing Data:** Cooked builds may have some MoverComponent data missing, resulting in a non-functional actor. Disable the actor's "Generate Optimized Blueprint Component Data" option (bOptimizeBPComponentData in C++). 
@@ -170,3 +194,4 @@ Additionally, when the simulation runs at a lower rate than rendering, the resul
 
 **Sim Blackboard Is Not Yet Rollback-Friendly:** When rollbacks occur, the blackboard's contents are simply invalidated. Although the blackboard is useful for avoiding repeating computations from a prior frame, movement logic cannot always rely on it having a valid entry and should have a fallback.
 
+**Multiple Modifiers of the Same Type:** Currently Mover expects there to be only one modifier of a given type active at a time. This can be worked around by overriding the Matches function and use more criteria than just the type, as long as that criteria is reflected by data included in the NetSerialize function.
diff --git a/Engine/Plugins/Experimental/MoverExamples/MoverExamples.uplugin b/Engine/Plugins/Experimental/MoverExamples/MoverExamples.uplugin
--- a/Engine/Plugins/Experimental/MoverExamples/MoverExamples.uplugin
+++ b/Engine/Plugins/Experimental/MoverExamples/MoverExamples.uplugin
@@ -3,7 +3,7 @@
 	"Version": 1,
 	"VersionName": "0.1",
 	"FriendlyName": "Mover Examples",
-	"Description": "Non-essential classes and content for the Mover plugin. Includes sample code, test maps, etc to help developers start using the system. Not intended to be used directly in a shipping product.",
+	"Description": "Non-essential classes and content for the Mover plugin. Includes sample code, test maps, etc to help developers start using the system. Not intended to be used directly in a shipping product. \nPlease refer to the Mover plugin's README document for information about getting started, an overview of concepts, and known issues.",
 	"Category": "Gameplay",
 	"CreatedBy": "Epic Games, Inc.",
 	"CreatedByURL": "https://epicgames.com",
