Mover: updates to readme about instanced subobjects. No functional changes.

#rnx

diff --git a/Engine/Plugins/Experimental/Mover/README.md b/Engine/Plugins/Experimental/Mover/README.md
--- a/Engine/Plugins/Experimental/Mover/README.md
+++ b/Engine/Plugins/Experimental/Mover/README.md
@@ -182,6 +182,8 @@ Additionally, when the simulation runs at a lower rate than rendering, the resul
 
 **Limited Blueprinting Support:** Blueprint functionality isn't 100% supported yet. There are certain things that still require native C++ code, or are clunky to implement in Blueprints.
 
+**Instanced Subobjects:** There are known issues with removing and replacing elements of the movement mode map on instanced actors placed in a level, potentially preventing them from moving. Consider other workflows that avoid this kind of customization on level-placed instances.
+
 **Arbitrary Gravity, Collision Shapes, etc. in the Default Movement Set:** Although the core MoverComponent tries to make as few requirements as possible on the composition of the actor, the default movement set has more rigid assumptions.  For example, the default movement set currently assumes a capsule shape. Some features such as arbitrary gravity are not fully supported in all cases yet.
 
 **Animation of Sim Proxy Example Characters:** Animation on another client's pawn (a sim proxy) may not be fully replicated during certain actions. This will be improved in a future release.
