Mover: readme updates for physics characters

#rnx

diff --git a/Engine/Plugins/Experimental/Mover/README.md b/Engine/Plugins/Experimental/Mover/README.md
--- a/Engine/Plugins/Experimental/Mover/README.md
+++ b/Engine/Plugins/Experimental/Mover/README.md
@@ -153,7 +153,8 @@ Other MoverExamples pawns, and the CharacterMovementComponent, use a "kinematic"
 This physics-based character is NOT using the Network Prediction plugin. Instead, it is driven by the Chaos Networked Physics system, which has many similarities in how it operates. Under the hood, a Character Ground Constraint from the physics system handles performing the proposed motion that the Mover system generates. The physics simulation is running at a fixed tick rate, ahead of the game thread. As the simulation progresses, the game thread's character representation interpolates towards the most recent physics state.
 
 Notes:
-- Due to the async nature of the physics simulation, there is some additional lag time between player input and affecting the pawn's movement seen on screen.
+- Due to the async nature of the physics simulation, there is some additional latency between movement and other gameplay systems. This will affect things like the amount of time between player input occurring and the pawn moving on screen. And when using animation root motion, there may be delay between a character's animation pose and when the motion is seen.
+- Some Mover mechanisms will not work with a physics-based character, such as the teleport effect.
 - Other physics objects that the character interacts with should be set up for replication. Otherwise this will lead to differences between client and server.
 - Interactions with moving non-physics objects will likely show poor results, due to ticking differences between the physics simulation and the rest of the game world.
 - Various gameplay events may not be connected or are unreliable. 
