Add spatial readiness API object. This object generates handles which can call back into it, but protects against invalid access with smart pointers, and uses delegates for functional injection.

diff --git a/Engine/Plugins/Experimental/SpatialReadiness/Source/PhysicsSpatialReadiness/Private/PhysicsSpatialReadiness/PhysicsSpatialReadinessSimCallback.cpp b/Engine/Plugins/Experimental/SpatialReadiness/Source/PhysicsSpatialReadiness/Private/PhysicsSpatialReadiness/PhysicsSpatialReadinessSimCallback.cpp
--- a/Engine/Plugins/Experimental/SpatialReadiness/Source/PhysicsSpatialReadiness/Private/PhysicsSpatialReadiness/PhysicsSpatialReadinessSimCallback.cpp
+++ b/Engine/Plugins/Experimental/SpatialReadiness/Source/PhysicsSpatialReadiness/Private/PhysicsSpatialReadiness/PhysicsSpatialReadinessSimCallback.cpp
@@ -1,6 +1,9 @@
 // Copyright Epic Games, Inc. All Rights Reserved.
 
+#include "PhysicsProxy/SingleParticlePhysicsProxy.h"
 #include "PhysicsSpatialReadinessSimCallback.h"
+#include "Chaos/ChaosEngineInterface.h"
+#include "SpatialReadinessHandle.h"
 
 using namespace Chaos;
 
@@ -15,3 +18,24 @@ void FPhysicsSpatialReadinessSimCallback::OnMidPhaseModification_Internal(FMidPh
 void FPhysicsSpatialReadinessSimCallback::OnPostIntegrate_Internal()
 {
 }
+
+FSpatialReadinessHandle FPhysicsSpatialReadinessSimCallback::AddReadinessVolume(const FBox& Bounds, const FString& Description)
+{
+	return SpatialReadiness.CreateHandle(Bounds, Description);
+}
+
+int32 FPhysicsSpatialReadinessSimCallback::AddUnreadyVolume_Internal(const FBox& Bounds, const FString& Description)
+{
+	//
+	// TODO
+	//
+
+	return INDEX_NONE;
+}
+
+void FPhysicsSpatialReadinessSimCallback::RemoveUnreadyVolume_Internal(int32 UnreadyVolumeIndex)
+{
+	//
+	// TODO
+	//
+}
diff --git a/Engine/Plugins/Experimental/SpatialReadiness/Source/PhysicsSpatialReadiness/Private/PhysicsSpatialReadiness/PhysicsSpatialReadinessSubsystem.cpp b/Engine/Plugins/Experimental/SpatialReadiness/Source/PhysicsSpatialReadiness/Private/PhysicsSpatialReadiness/PhysicsSpatialReadinessSubsystem.cpp
--- a/Engine/Plugins/Experimental/SpatialReadiness/Source/PhysicsSpatialReadiness/Private/PhysicsSpatialReadiness/PhysicsSpatialReadinessSubsystem.cpp
+++ b/Engine/Plugins/Experimental/SpatialReadiness/Source/PhysicsSpatialReadiness/Private/PhysicsSpatialReadiness/PhysicsSpatialReadinessSubsystem.cpp
@@ -32,45 +32,33 @@ void UPhysicsSpatialReadinessSubsystem::Deinitialize()
 	Super::Deinitialize();
 }
 
-FUnreadyVolumeHandle UPhysicsSpatialReadinessSubsystem::AddUnreadyVolume(const FBox& AABB, const FString& Description)
+FSpatialReadinessHandle UPhysicsSpatialReadinessSubsystem::AddReadinessVolume(const FBox& Bounds, const FString& Description)
 {
-	/*
-	* TODO:
-	* - Create geometry particle for unready volume
-	* - Add particle to list of unready particles in sim callback
-	* - Add particle to our internal array of unready particles at first free index
-	* - Return the index that the particle was stored at as the "handle" to the volume
-	*/
-	return FUnreadyVolumeHandle();
-}
+	if (ensureMsgf(SimCallback, TEXT("Attempted to add readiness volume before sim callback was initialized")))
+	{
+		return SimCallback->AddReadinessVolume(Bounds, Description);
+	}
 
-void UPhysicsSpatialReadinessSubsystem::RemoveUnreadyVolume(FUnreadyVolumeHandle Handle)
-{
-	/*
-	* TODO:
-	* - Destroy geometry particle corresponding to volume handle
-	* - Free index
-	*/
+	// It's a sad day when we've gotta debug this case
+	return FSpatialReadinessHandle();
 }
 
 bool UPhysicsSpatialReadinessSubsystem::CreateSimCallback()
 {
-	// In case we already have one, destroy it
-	DestroySimCallback();
+	if (SimCallback)
+	{
+		DestroySimCallback();
+	}
 
-	if (SimCallback == nullptr)
+	if (ensureMsgf(SimCallback, TEXT("Tried and failed to destroy existing sim callback so that a new one could be created.")))
 	{
-		if (UWorld* World = GetWorld())
-		{
-			if (FPhysScene* PhysScene = World->GetPhysicsScene())
-			{
-				if (FPhysicsSolverBase* Solver = PhysScene->GetSolver())
-				{
-					SimCallback = Solver->CreateAndRegisterSimCallbackObject_External<FPhysicsSpatialReadinessSimCallback>();
-					return true;
-				}
-			}
-		}
+		return false;
+	}
+
+	if (Chaos::FPhysicsSolver* Solver = GetSolver())
+	{
+		SimCallback = Solver->CreateAndRegisterSimCallbackObject_External<FPhysicsSpatialReadinessSimCallback>();
+		return true;
 	}
 
 	return false;
@@ -78,21 +66,28 @@ bool UPhysicsSpatialReadinessSubsystem::CreateSimCallback()
 
 bool UPhysicsSpatialReadinessSubsystem::DestroySimCallback()
 {
-	if (SimCallback != nullptr)
+	if (Chaos::FPhysicsSolver* Solver = GetSolver())
 	{
-		if (UWorld* World = GetWorld())
+		Solver->UnregisterAndFreeSimCallbackObject_External(SimCallback);
+		SimCallback = nullptr;
+		return true;
+	}
+
+	return false;
+}
+
+Chaos::FPhysicsSolver* UPhysicsSpatialReadinessSubsystem::GetSolver()
+{
+	if (UWorld* World = GetWorld())
+	{
+		if (FPhysScene* Scene = World->GetPhysicsScene())
 		{
-			if (FPhysScene* PhysScene = World->GetPhysicsScene())
+			if (Chaos::FPhysicsSolver* Solver = Scene->GetSolver())
 			{
-				if (FPhysicsSolverBase* Solver = PhysScene->GetSolver())
-				{
-					Solver->UnregisterAndFreeSimCallbackObject_External(SimCallback);
-					SimCallback = nullptr;
-					return true;
-				}
+				return Solver;
 			}
 		}
 	}
 
-	return false;
+	return nullptr;
 }
diff --git a/Engine/Plugins/Experimental/SpatialReadiness/Source/PhysicsSpatialReadiness/Public/PhysicsSpatialReadinessSimCallback.h b/Engine/Plugins/Experimental/SpatialReadiness/Source/PhysicsSpatialReadiness/Public/PhysicsSpatialReadinessSimCallback.h
--- a/Engine/Plugins/Experimental/SpatialReadiness/Source/PhysicsSpatialReadiness/Public/PhysicsSpatialReadinessSimCallback.h
+++ b/Engine/Plugins/Experimental/SpatialReadiness/Source/PhysicsSpatialReadiness/Public/PhysicsSpatialReadinessSimCallback.h
@@ -5,18 +5,9 @@
 #include "Chaos/SimCallbackObject.h"
 #include "PhysicsProxy/SingleParticlePhysicsProxy.h"
 #include "Chaos/Framework/HashMappedArray.h"
+#include "SpatialReadinessAPI.h"
 
 using namespace Chaos;
-/*
-using Chaos::FSimCallbackInput;
-using Chaos::FSimCallbackOutput;
-using Chaos::FSimCallbackNoInput;
-using Chaos::FSimCallbackNoOutput;
-using Chaos::TSimCallbackObject;
-using Chaos::FSingleParticlePhysicsProxy;
-using Chaos::FGeometryParticleHandle;
-using Chaos::FMidPhaseModifierAccessor;
-*/
 
 struct FPhysicsSpatialReadinessSimCallbackInput : public FSimCallbackInput
 {
@@ -30,45 +21,67 @@ struct FPhysicsSpatialReadinessSimCallbackInput : public FSimCallbackInput
 	}
 };
 
-struct FPhysicsSpatialReadinessSimCallback : public TSimCallbackObject<
-	FPhysicsSpatialReadinessSimCallbackInput,
-	FSimCallbackNoOutput,
-	ESimCallbackOptions::Presimulate |
-	ESimCallbackOptions::MidPhaseModification |
-	ESimCallbackOptions::PostIntegrate>
+struct FPhysicsSpatialReadinessSimCallback
+	: public TSimCallbackObject<
+		FPhysicsSpatialReadinessSimCallbackInput,
+		FSimCallbackNoOutput,
+		ESimCallbackOptions::Presimulate |
+		ESimCallbackOptions::MidPhaseModification |
+		ESimCallbackOptions::PostIntegrate>
 {
+	using This = FPhysicsSpatialReadinessSimCallback;
+
 public:
 	FPhysicsSpatialReadinessSimCallback()
-		: HashSize(1024)
-		, ParticleDataCache(HashSize)
+		: SpatialReadiness(this, &This::AddUnreadyVolume_Internal, &This::RemoveUnreadyVolume_Internal)
+		, HashSize(1024)
+		, ParticleDataCache_PT(HashSize)
 	{ }
 
-protected:
-
 	/* begin: TSimCallbackObject */
+protected:
 	virtual void OnPreSimulate_Internal() override;
 	virtual void OnMidPhaseModification_Internal(FMidPhaseModifierAccessor& Accessor) override;
 	virtual void OnPostIntegrate_Internal() override;
 	/* end: TSimCallbackObject */
 
+public:
+	FSpatialReadinessHandle AddReadinessVolume(const FBox& Bounds, const FString& Description);
+
 private:
 
+	FSpatialReadinessAPI SpatialReadiness;
+	int32 AddUnreadyVolume_Internal(const FBox& Bounds, const FString& Description);
+	void RemoveUnreadyVolume_Internal(int32 UnreadyVolumeIndex);
+
+	// List of unready volume physics proxies. We directly use single particle
+	// physics proxy rather than something more generic because we know that
+	// we are only going to create static single particles for these volumes. 
+	TSparseArray<FSingleParticlePhysicsProxy*> UnreadyVolumeParticles_GT;
+
 	// List of particle handles which represent unready volumes
-	TSet<FGeometryParticleHandle*> UnreadyVolumes;
+	TSet<FGeometryParticleHandle*> UnreadyVolumeParticles_PT;
 
 	// Size of the ParticleDataCache
-	uint32 HashSize;
+	const uint32 HashSize;
 
 	// This type data is used to create per-particle readiness data.
 	struct FParticleDataCacheTraits
 	{
-		static uint32 GetIDHash(const FGeometryParticleHandle* Particle) { return MurmurFinalize32(uint32(Particle->UniqueIdx().Idx)); }
-		static bool ElementHasID(const FGeometryParticleHandle* A, const FGeometryParticleHandle* B) { return A == B; }
+		static uint32 GetIDHash(const FGeometryParticleHandle* Particle)
+		{
+			return MurmurFinalize32(uint32(Particle->UniqueIdx().Idx));
+		}
+		static bool ElementHasID(const FGeometryParticleHandle* A, const FGeometryParticleHandle* B)
+		{
+			return A == B;
+		}
 	};
+
 	using FParticleDataCache = Private::THashMappedArray<FGeometryParticleHandle*, EObjectStateType, FParticleDataCacheTraits>;
 
 	// "Unready" particles are forced to be stationary in PreSimulate, and restored
 	// to their previous state in PostIntegrate. The values needed for restoration
 	// are stored and mapped back to particle id's in this hash-mapped array.
-	FParticleDataCache ParticleDataCache;
+	FParticleDataCache ParticleDataCache_PT;
 };
diff --git a/Engine/Plugins/Experimental/SpatialReadiness/Source/PhysicsSpatialReadiness/Public/PhysicsSpatialReadinessSubsystem.h b/Engine/Plugins/Experimental/SpatialReadiness/Source/PhysicsSpatialReadiness/Public/PhysicsSpatialReadinessSubsystem.h
--- a/Engine/Plugins/Experimental/SpatialReadiness/Source/PhysicsSpatialReadiness/Public/PhysicsSpatialReadinessSubsystem.h
+++ b/Engine/Plugins/Experimental/SpatialReadiness/Source/PhysicsSpatialReadiness/Public/PhysicsSpatialReadinessSubsystem.h
@@ -4,34 +4,30 @@
 
 #include "CoreMinimal.h"
 #include "Subsystems/WorldSubsystem.h"
-#include "SpatialReadinessInterface.h"
+#include "SpatialReadinessAPI.h"
 #include "PhysicsSpatialReadinessSimCallback.h"
 #include "PhysicsSpatialReadinessSubsystem.generated.h"
 
 UCLASS()
-class UPhysicsSpatialReadinessSubsystem : public UWorldSubsystem, public ISpatialReadiness
+class UPhysicsSpatialReadinessSubsystem : public UWorldSubsystem
 {
 	GENERATED_BODY()
 
+	/* begin: UWorldSubsystem */
 protected:
-
-	// Begin: UWorldSubsystem
 	virtual bool ShouldCreateSubsystem(UObject* Outer) const override;
 	virtual void OnWorldBeginPlay(UWorld& InWorld) override;
 	virtual void Initialize(FSubsystemCollectionBase& Collection) override;
 	virtual void Deinitialize() override;
-	// End: UWorldSubsystem
+	/* end: UWorldSubsystem */
 
-	// Begin: ISpatialReadiness
-	virtual FUnreadyVolumeHandle AddUnreadyVolume(const FBox& AABB, const FString& Description) override;
-	virtual void RemoveUnreadyVolume(FUnreadyVolumeHandle Handle) override;
-	// End: ISpatialReadiness
+public:
+	FSpatialReadinessHandle AddReadinessVolume(const FBox& Bounds, const FString& Description);
 
 private:
 	bool CreateSimCallback();
 	bool DestroySimCallback();
+	Chaos::FPhysicsSolver* GetSolver();
 
 	FPhysicsSpatialReadinessSimCallback* SimCallback = nullptr;
-
-	TSparseArray<FSingleParticlePhysicsProxy*> UnreadyVolumes;
 };
diff --git a/Engine/Plugins/Experimental/SpatialReadiness/Source/SpatialReadiness/Private/SpatialReadiness/SpatialReadinessAPI.cpp b/Engine/Plugins/Experimental/SpatialReadiness/Source/SpatialReadiness/Private/SpatialReadiness/SpatialReadinessAPI.cpp
new file mode 100644
--- /dev/null
+++ b/Engine/Plugins/Experimental/SpatialReadiness/Source/SpatialReadiness/Private/SpatialReadiness/SpatialReadinessAPI.cpp
@@ -0,0 +1,36 @@
+// Copyright Epic Games, Inc. All Rights Reserved.
+
+#include "SpatialReadinessAPI.h"
+#include "SpatialReadinessAPIDelegates.h"
+
+FSpatialReadinessAPI::FSpatialReadinessAPI(
+	FAddUnreadyVolume_Function AddUnreadyVolume,
+	FRemoveUnreadyVolume_Function RemoveUnreadyVolume)
+	: Delegates(MakeShared<FSpatialReadinessAPIDelegates>(AddUnreadyVolume, RemoveUnreadyVolume))
+{ }
+
+FSpatialReadinessAPI::~FSpatialReadinessAPI()
+{
+	if (ensureMsgf(Delegates, TEXT("An internal spatial readiness interfaec was freed before the destruction of its outer")))
+	{
+		// Some handle might still have reference to the internal since it has a WeakPtr
+		// which can be converted to shared. That means it might try to access its
+		// delegates, which will now likely be bound to invalid functions.
+		Delegates->AddUnreadyVolumeDelegate.Unbind();
+		Delegates->RemoveUnreadyVolumeDelegate.Unbind();
+	}
+}
+
+FSpatialReadinessHandle FSpatialReadinessAPI::CreateHandle(const FBox& Bounds, const FString& Description)
+{
+	// Create a handle for a new unready volume
+	FSpatialReadinessHandle Handle(Delegates.ToWeakPtr(), Bounds, Description);
+
+	// Default to an unready state. This should call back into AddUnreadyVolume_Internal
+	// to create the underlying unready volume.
+	Handle.MarkUnready();
+
+	// Return the newly created handle. This handle will only be valid so long as this
+	// ISpatialReadiness object is in scope.
+	return Handle;
+}
diff --git a/Engine/Plugins/Experimental/SpatialReadiness/Source/SpatialReadiness/Private/SpatialReadiness/SpatialReadinessAPIDelegates.cpp b/Engine/Plugins/Experimental/SpatialReadiness/Source/SpatialReadiness/Private/SpatialReadiness/SpatialReadinessAPIDelegates.cpp
new file mode 100644
--- /dev/null
+++ b/Engine/Plugins/Experimental/SpatialReadiness/Source/SpatialReadiness/Private/SpatialReadiness/SpatialReadinessAPIDelegates.cpp
@@ -0,0 +1,23 @@
+// Copyright Epic Games, Inc. All Rights Reserved.
+
+#include "SpatialReadinessAPIDelegates.h"
+#include "SpatialReadinessStats.h"
+
+FSpatialReadinessAPIDelegates::FSpatialReadinessAPIDelegates(
+	FAddUnreadyVolume_Function AddUnreadyVolume,
+	FRemoveUnreadyVolume_Function RemoveUnreadyVolume)
+{
+	// Each of these bindings captures _itself_ so it can check to see that it's still
+	// bound after execution. This somewhat odd construction exists because this is the
+	// only point at which we can't be sure that the interface that created the handle
+	// still exists... although the only way it should be possible for it to have been
+	// deleted mid-execution of one of these lambdas would be if there was some crazy
+	// threading going on.
+	//
+	// Here's the situation that we're trying to protect against:
+	// 1) Interface.CreateHandle(); // called on thread 0, triggers call to AddUnreadyVolumeDelegate.Execute(...)
+	// 2) Interface.~(); // called on thread 1, before above delegate execution completes
+
+	AddUnreadyVolumeDelegate.BindLambda(AddUnreadyVolume);
+	RemoveUnreadyVolumeDelegate.BindLambda(RemoveUnreadyVolume);
+}
diff --git a/Engine/Plugins/Experimental/SpatialReadiness/Source/SpatialReadiness/Private/SpatialReadiness/SpatialReadinessAPIDelegates.h b/Engine/Plugins/Experimental/SpatialReadiness/Source/SpatialReadiness/Private/SpatialReadiness/SpatialReadinessAPIDelegates.h
new file mode 100644
--- /dev/null
+++ b/Engine/Plugins/Experimental/SpatialReadiness/Source/SpatialReadiness/Private/SpatialReadiness/SpatialReadinessAPIDelegates.h
@@ -0,0 +1,19 @@
+// Copyright Epic Games, Inc. All Rights Reserved.
+
+#pragma once
+
+#include "CoreMinimal.h"
+#include "SpatialReadinessSignatures.h"
+
+struct FSpatialReadinessAPIDelegates final
+{
+	// Constructor which binds TFunctions
+	FSpatialReadinessAPIDelegates(
+		FAddUnreadyVolume_Function AddUnreadyVolume,
+		FRemoveUnreadyVolume_Function RemoveUnreadyVolume);
+
+	~FSpatialReadinessAPIDelegates() = default;
+
+	FAddUnreadyVolume_Delegate AddUnreadyVolumeDelegate;
+	FRemoveUnreadyVolume_Delegate RemoveUnreadyVolumeDelegate;
+};
diff --git a/Engine/Plugins/Experimental/SpatialReadiness/Source/SpatialReadiness/Private/SpatialReadiness/SpatialReadinessHandle.cpp b/Engine/Plugins/Experimental/SpatialReadiness/Source/SpatialReadiness/Private/SpatialReadiness/SpatialReadinessHandle.cpp
new file mode 100644
--- /dev/null
+++ b/Engine/Plugins/Experimental/SpatialReadiness/Source/SpatialReadiness/Private/SpatialReadiness/SpatialReadinessHandle.cpp
@@ -0,0 +1,136 @@
+// Copyright Epic Games, Inc. All Rights Reserved.
+
+#include "SpatialReadinessHandle.h"
+#include "SpatialReadinessAPIDelegates.h"
+
+FSpatialReadinessHandle::FSpatialReadinessHandle(
+	TWeakPtr<FSpatialReadinessAPIDelegates> InDelegates,
+	const FBox& InBounds,
+	const FString& InDescription)
+	: Delegates(InDelegates)
+	, Index(INDEX_NONE)
+	, Bounds(InBounds)
+#if WITH_SPATIALREADINESS_HANDLE_DESCRIPTIONS
+	, Description(InDescription)
+#endif
+{ }
+
+FSpatialReadinessHandle::FSpatialReadinessHandle(FSpatialReadinessHandle&& Other)
+	: Delegates(Other.Delegates)
+	, Index(Other.Index)
+	, Bounds(Other.Bounds)
+#if WITH_SPATIALREADINESS_HANDLE_DESCRIPTIONS
+	, Description(Other.Description)
+#endif
+{
+	Other.Delegates.Reset();
+	Other.Index = INDEX_NONE;
+	Other.Bounds = FBox();
+#if WITH_SPATIALREADINESS_HANDLE_DESCRIPTIONS
+	Other.Description = TEXT("");
+#endif
+}
+
+FSpatialReadinessHandle::~FSpatialReadinessHandle()
+{
+	// When a handle is destroyed, we mark the corresponding volume as
+	// "ready" (in case it wasn't already). This ensures that we do not
+	// end up losing track of unready volumes and leaving them hanging
+	// around in an acceleration structure.
+	//
+	// This is the only case where we actually do an IsValid check
+	// without _ensuring_ validity. If the api which created this
+	// handle has already been deleted, then MarkReady is not going to
+	// do anything anyway...
+	//
+	// I'm a bit conflicted on adding this check because if we ever
+	// come in here when IsValid == false, that means a handle was held
+	// onto for too long.
+	if (IsValid())
+	{
+		MarkReady();
+	}
+}
+
+bool FSpatialReadinessHandle::IsValid() const
+{
+	return Delegates.IsValid();
+}
+
+bool FSpatialReadinessHandle::EnsureIsValid() const
+{
+	return ensureMsgf(IsValid(), TEXT("Readiness volume handle with invalid delegates pointer is being accessed"));
+}
+
+bool FSpatialReadinessHandle::IsReady() const
+{
+	EnsureIsValid();
+
+	// If we have no index, that means there's no underlying unready volume,
+	// so we're "ready"
+	return Index == INDEX_NONE;
+}
+
+void FSpatialReadinessHandle::MarkReady()
+{
+	if (IsReady())
+	{
+		return;
+	}
+
+	TSharedPtr<FSpatialReadinessAPIDelegates> SharedDelegates = Delegates.Pin();
+	if (ensureMsgf(SharedDelegates, TEXT("Tried to mark ready a readiness volume handle with an invalid delegates container")))
+	{
+		// Remove the unready volume corresponding with our current index
+		SharedDelegates->RemoveUnreadyVolumeDelegate.ExecuteIfBound(Index);
+	}
+
+	// Stop tracking this index
+	Index = INDEX_NONE;
+}
+
+void FSpatialReadinessHandle::MarkUnready()
+{
+	if (!IsReady())
+	{
+		return;
+	}
+
+	TSharedPtr<FSpatialReadinessAPIDelegates> SharedDelegates = Delegates.Pin();
+	if (ensureMsgf(SharedDelegates, TEXT("Tried to mark unready a readiness volume handle with an invalid delegates container")))
+	{
+		// If we've got an index, that means we're already associated with an
+		// existing unready volume, which must be destroyed.
+		if (Index != INDEX_NONE)
+		{
+			// If we already had an unread volume, remove it so we can make
+			// a new one without leaving the old one dangling
+			SharedDelegates->RemoveUnreadyVolumeDelegate.ExecuteIfBound(Index);
+		}
+
+		// Update our index to a new unready volume
+		if (SharedDelegates->AddUnreadyVolumeDelegate.IsBound())
+		{
+			Index = SharedDelegates->AddUnreadyVolumeDelegate.Execute(Bounds, GetDescription());
+		}
+	}
+}
+
+void FSpatialReadinessHandle::MarkUnready(FBox& NewBounds)
+{
+	// Store new bounds
+	Bounds = NewBounds;
+
+	// Create a new unready volume with new bounds
+	MarkUnready();
+}
+
+const FString& FSpatialReadinessHandle::GetDescription() const
+{
+#if WITH_SPATIALREADINESS_HANDLE_DESCRIPTIONS
+	return Description;
+#else
+	static const FString EmptyString;
+	return EmptyString;
+#endif
+}
diff --git a/Engine/Plugins/Experimental/SpatialReadiness/Source/SpatialReadiness/Private/SpatialReadiness/SpatialReadinessStats.h b/Engine/Plugins/Experimental/SpatialReadiness/Source/SpatialReadiness/Private/SpatialReadiness/SpatialReadinessStats.h
new file mode 100644
--- /dev/null
+++ b/Engine/Plugins/Experimental/SpatialReadiness/Source/SpatialReadiness/Private/SpatialReadiness/SpatialReadinessStats.h
@@ -0,0 +1,12 @@
+// Copyright Epic Games, Inc. All Rights Reserved.
+
+#pragma once
+
+#include "Stats/Stats.h"
+
+DECLARE_STATS_GROUP(TEXT("SpatialReadiness"), STATGROUP_SpatialReadiness, STATCAT_Advanced);
+DECLARE_CYCLE_STAT(TEXT(""), STAT_SpatialReadiness_AddUnreadyVolume, STATGROUP_SpatialReadiness);
+DECLARE_CYCLE_STAT(TEXT(""), STAT_SpatialReadiness_RemoveUnreadyVolume, STATGROUP_SpatialReadiness);
+
+
+
diff --git a/Engine/Plugins/Experimental/SpatialReadiness/Source/SpatialReadiness/Private/SpatialReadiness/UnreadyVolumeHandle.cpp b/Engine/Plugins/Experimental/SpatialReadiness/Source/SpatialReadiness/Private/SpatialReadiness/UnreadyVolumeHandle.cpp
deleted file mode 100644
diff --git a/Engine/Plugins/Experimental/SpatialReadiness/Source/SpatialReadiness/Public/SpatialReadinessAPI.h b/Engine/Plugins/Experimental/SpatialReadiness/Source/SpatialReadiness/Public/SpatialReadinessAPI.h
new file mode 100644
--- /dev/null
+++ b/Engine/Plugins/Experimental/SpatialReadiness/Source/SpatialReadiness/Public/SpatialReadinessAPI.h
@@ -0,0 +1,47 @@
+// Copyright Epic Games, Inc. All Rights Reserved.
+
+#pragma once
+
+#include "CoreMinimal.h"
+#include "Templates/Function.h"
+#include "SpatialReadinessHandle.h"
+#include "SpatialReadinessSignatures.h"
+
+struct FSpatialReadinessAPIDelegates;
+
+class FSpatialReadinessAPI final
+{
+public:
+
+	// This constructor binds TFunctions to the internal add and remove delegates
+	SPATIALREADINESS_API FSpatialReadinessAPI(
+		FAddUnreadyVolume_Function AddUnreadyVolume,
+		FRemoveUnreadyVolume_Function RemoveUnreadyVolume);
+
+	// This constructor binds member functions to the internall add and remove delegates.
+	//
+	// Please be sure that this API object will not outlive the object to whose members
+	// its delegates are bound after this construction!
+	template <typename ObjectT>
+	FSpatialReadinessAPI(ObjectT* Object,
+		TAddUnreadyVolume_Member<ObjectT> AddUnreadyVolume,
+		TRemoveUnreadyVolume_Member<ObjectT> RemoveUnreadyVolume);
+
+	SPATIALREADINESS_API ~FSpatialReadinessAPI();
+
+	SPATIALREADINESS_API FSpatialReadinessHandle CreateHandle(const FBox& Bounds, const FString& Description);
+
+private:
+	TSharedPtr<FSpatialReadinessAPIDelegates> Delegates;
+};
+
+template <typename ObjectT>
+FSpatialReadinessAPI::FSpatialReadinessAPI(ObjectT* Object,
+	TAddUnreadyVolume_Member<ObjectT> AddUnreadyVolume,
+	TRemoveUnreadyVolume_Member<ObjectT> RemoveUnreadyVolume)
+	: FSpatialReadinessAPI(
+		[Object, AddUnreadyVolume] (const FBox& Bounds, const FString& Desc) -> int32
+		{ return (Object->*AddUnreadyVolume)(Bounds, Desc); },
+		[Object, RemoveUnreadyVolume] (int32 Index) -> void
+		{ (Object->*RemoveUnreadyVolume)(Index); })
+{ }
\ No newline at end of file
diff --git a/Engine/Plugins/Experimental/SpatialReadiness/Source/SpatialReadiness/Public/SpatialReadinessHandle.h b/Engine/Plugins/Experimental/SpatialReadiness/Source/SpatialReadiness/Public/SpatialReadinessHandle.h
new file mode 100644
--- /dev/null
+++ b/Engine/Plugins/Experimental/SpatialReadiness/Source/SpatialReadiness/Public/SpatialReadinessHandle.h
@@ -0,0 +1,83 @@
+// Copyright Epic Games, Inc. All Rights Reserved.
+
+#pragma once
+
+#include "CoreMinimal.h"
+#include "Math/MathFwd.h"
+
+#ifndef WITH_SPATIALREADINESS_HANDLE_DESCRIPTIONS
+#define WITH_SPATIALREADINESS_HANDLE_DESCRIPTIONS (!UE_BUILD_SHIPPING)
+#endif
+
+struct FSpatialReadinessAPIDelegates;
+
+class FSpatialReadinessHandle final
+{
+public:
+
+	// Check to make sure that the readiness system which created this handle
+	// did not go out of scope.
+	SPATIALREADINESS_API bool IsValid() const;
+
+	// Ensure that IsValid() returns true, printing a message and returning
+	// false if not.
+	SPATIALREADINESS_API bool EnsureIsValid() const;
+
+	// Check to see if this handle points to a volume which is currently "ready".
+	SPATIALREADINESS_API bool IsReady() const;
+
+	// Mark this volume as "ready"
+	SPATIALREADINESS_API void MarkReady();
+
+	// Mark this volume as "unready"
+	SPATIALREADINESS_API void MarkUnready();
+
+	// Optional override to MakeUnready() which adjusts the bounds
+	SPATIALREADINESS_API void MarkUnready(FBox& NewBounds);
+
+	// Get the description of this handle - will be empty string if
+	// the WITH_SPATIALREADINESS_HANDLE_DESCRIPTIONS macro is 0
+	SPATIALREADINESS_API const FString& GetDescription() const;
+
+	// When a handle is destroyed, it will make the corresponding volume as ready
+	SPATIALREADINESS_API ~FSpatialReadinessHandle();
+
+	// The default constructor will create an invalid handle
+	SPATIALREADINESS_API FSpatialReadinessHandle() = default;
+
+	// Readiness handles are unique and cannot be copied
+	FSpatialReadinessHandle(const FSpatialReadinessHandle&) = delete;
+
+	// Move constructor
+	SPATIALREADINESS_API FSpatialReadinessHandle(FSpatialReadinessHandle&& Other);
+
+private:
+
+	// We keep the constructor private and befriend the API object to restrict
+	// the means by which a handle can be created.
+	friend class FSpatialReadinessAPI;
+	FSpatialReadinessHandle(
+		TWeakPtr<FSpatialReadinessAPIDelegates> InDelegates,
+		const FBox& InBounds,
+		const FString& InDescription);
+
+	// Ref to the spatial readiness system. It will be an access error to use
+	// this handle if this interface ever goes out of scope
+	TWeakPtr<FSpatialReadinessAPIDelegates> Delegates;
+
+	// Index used by the interface to select the related volume, if
+	// one exists (ie if this handle refers to an "unready" state). If one
+	// doesn't exist, then this handle is "ready" and this is INDEX_NONE
+	int32 Index;
+
+	// Bounds of the unready volume. This is used when marking the volume
+	// as unready to create the internal "unready volume".
+	FBox Bounds;
+
+#if WITH_SPATIALREADINESS_HANDLE_DESCRIPTIONS
+	// Description of the reason why the volume referred to by this handle
+	// might be unready. Again, this is used only when creating the internal
+	// unready volume.
+	FString Description;
+#endif
+};
diff --git a/Engine/Plugins/Experimental/SpatialReadiness/Source/SpatialReadiness/Public/SpatialReadinessInterface.h b/Engine/Plugins/Experimental/SpatialReadiness/Source/SpatialReadiness/Public/SpatialReadinessInterface.h
deleted file mode 100644
diff --git a/Engine/Plugins/Experimental/SpatialReadiness/Source/SpatialReadiness/Public/SpatialReadinessSignatures.h b/Engine/Plugins/Experimental/SpatialReadiness/Source/SpatialReadiness/Public/SpatialReadinessSignatures.h
new file mode 100644
--- /dev/null
+++ b/Engine/Plugins/Experimental/SpatialReadiness/Source/SpatialReadiness/Public/SpatialReadinessSignatures.h
@@ -0,0 +1,20 @@
+// Copyright Epic Games, Inc. All Rights Reserved.
+
+#pragma once
+
+#include "CoreMinimal.h"
+#include "GenericPlatform/GenericPlatformMisc.h"
+
+// Declare internal delegate types
+DECLARE_DELEGATE_RetVal_TwoParams(int32, FAddUnreadyVolume_Delegate, const FBox&, const FString&);
+DECLARE_DELEGATE_OneParam(FRemoveUnreadyVolume_Delegate, int32);
+
+// Declare TFunction types which can be bound to the delegates
+using FAddUnreadyVolume_Function = TFunction<int32(const FBox&, const FString&)>;
+using FRemoveUnreadyVolume_Function = TFunction<void(int32)>;
+
+// Declare member function types which can be bound to delegates
+template <typename ObjectT>
+using TAddUnreadyVolume_Member = int32(ObjectT::*)(const FBox&, const FString&);
+template <typename ObjectT>
+using TRemoveUnreadyVolume_Member = void(ObjectT::*)(int32);
diff --git a/Engine/Plugins/Experimental/SpatialReadiness/Source/SpatialReadiness/Public/UnreadyVolumeHandle.h b/Engine/Plugins/Experimental/SpatialReadiness/Source/SpatialReadiness/Public/UnreadyVolumeHandle.h
deleted file mode 100644
diff --git a/Engine/Plugins/Experimental/SpatialReadiness/Tests/Private/PhysicsSpatialReadinessTests.cpp b/Engine/Plugins/Experimental/SpatialReadiness/Tests/Private/PhysicsSpatialReadinessTests.cpp
--- a/Engine/Plugins/Experimental/SpatialReadiness/Tests/Private/PhysicsSpatialReadinessTests.cpp
+++ b/Engine/Plugins/Experimental/SpatialReadiness/Tests/Private/PhysicsSpatialReadinessTests.cpp
@@ -1,25 +1,47 @@
 // Copyright Epic Games, Inc. All Rights Reserved.
 #include "TestHarness.h"
 #include "CoreMinimal.h"
-//#include "Chaos/Framework/ChaosUserDataPT.h"
 #include "PBDRigidsSolver.h"
 #include "ChaosSolversModule.h"
 #include "PhysicsSpatialReadinessSimCallback.h"
 
+namespace
+{
+	void AdvanceAndWait(Chaos::FPBDRigidsSolver* InSolver, float DeltaTime = 1)
+	{
+		if (InSolver)
+		{
+			InSolver->AdvanceAndDispatch_External(DeltaTime);
+			InSolver->WaitOnPendingTasks_External();
+		}
+	}
+}
 
 TEST_CASE("Sim Callback", "[physics spatial readiness]")
 {
-	const float DeltaTime = 1.f/30.f;
-
 	// Create a solver in the solvers module
 	FChaosSolversModule* Module = FChaosSolversModule::GetModule();
-	Chaos::FPBDRigidsSolver* Solver = Module->CreateSolver(nullptr, /*AsyncDt=*/DeltaTime, Chaos::EThreadingMode::TaskGraph);
-
-	// Create a test PhysicsSpatialReadinessSimCallback
-	FPhysicsSpatialReadinessSimCallback* SimCallback = Solver->CreateAndRegisterSimCallbackObject_External<FPhysicsSpatialReadinessSimCallback>();
+	Chaos::FPBDRigidsSolver* Solver = Module->CreateSolver(nullptr, /*AsyncDt=*/1, Chaos::EThreadingMode::TaskGraph);
 
-	SECTION("Create unready volume")
+	SECTION("Create unready volume, and let the readiness system go out of scope")
 	{
-		// TODO: Everything
+		// Create a test PhysicsSpatialReadinessSimCallback
+		FPhysicsSpatialReadinessSimCallback* SimCallback = Solver->CreateAndRegisterSimCallbackObject_External<FPhysicsSpatialReadinessSimCallback>();
+		AdvanceAndWait(Solver);
+
+		// Make a volume
+		const FBox Box(FVector(-1), FVector(1));
+		FSpatialReadinessHandle Handle = SimCallback->AddReadinessVolume(Box, TEXT("Test Volume"));
+		AdvanceAndWait(Solver);
+
+		// The volume handle should be "valid"
+		REQUIRE(Handle.IsValid());
+
+		// Destroy the sim callback
+		Solver->UnregisterAndFreeSimCallbackObject_External(SimCallback);
+		AdvanceAndWait(Solver);
+
+		// The volume handle should now be "invalid"
+		REQUIRE(!Handle.IsValid());
 	}
 }
\ No newline at end of file
diff --git a/Engine/Plugins/Experimental/SpatialReadiness/Tests/Private/SpatialReadinessTests.cpp b/Engine/Plugins/Experimental/SpatialReadiness/Tests/Private/SpatialReadinessTests.cpp
--- a/Engine/Plugins/Experimental/SpatialReadiness/Tests/Private/SpatialReadinessTests.cpp
+++ b/Engine/Plugins/Experimental/SpatialReadiness/Tests/Private/SpatialReadinessTests.cpp
@@ -1,11 +1,152 @@
 // Copyright Epic Games, Inc. All Rights Reserved.
 #include "TestHarness.h"
+#include "SpatialReadinessAPI.h"
 
-TEST_CASE("...", "[spatial readiness]")
+namespace
 {
-	SECTION("...")
+	// Make some test data for populating handles
+	static const FBox TestBox(FVector(-.5f), FVector(.5f));
+	static const FString TestDesc(TEXT("test volume"));
+}
+
+TEST_CASE("Spatial readiness api", "[spatial readiness]")
+{
+	// Make some lambdas that we can use to create spatial readiness
+	// api objects that do nothing.
+	struct FVolume { FBox Bounds; FString Desc; };
+	TSparseArray<FVolume> Volumes;
+	const auto Stub_AddVolume = [&Volumes](const FBox& Bounds, const FString& Desc) -> int32
+	{
+		int32 Index = 0;
+		return Volumes.EmplaceAtLowestFreeIndex(Index, Bounds, Desc);
+	};
+
+	const auto Stub_RemoveVolume = [&Volumes](int32 Index) -> void
+	{
+		Volumes.RemoveAt(Index);
+	};
+
+	// Make a spatial readiness api
+	TUniquePtr<FSpatialReadinessAPI> SpatialReadiness = MakeUnique<FSpatialReadinessAPI>(Stub_AddVolume, Stub_RemoveVolume);
+
+	SECTION("Volume is initially unready")
+	{
+		FSpatialReadinessHandle Handle = SpatialReadiness->CreateHandle(TestBox, TestDesc);
+		REQUIRE(Handle.IsReady() == false);
+	}
+
+	SECTION("Volume is removed when handle is marked ready")
 	{
-		// TODO: everything
-		REQUIRE(true);
+		FSpatialReadinessHandle Handle = SpatialReadiness->CreateHandle(TestBox, TestDesc);
+		REQUIRE(Volumes.Num() == 1);
+		Handle.MarkReady();
+		REQUIRE(Volumes.Num() == 0);
+	}
+
+	SECTION("Volume is added when handle is marked unready")
+	{
+		FSpatialReadinessHandle Handle = SpatialReadiness->CreateHandle(TestBox, TestDesc);
+		Handle.MarkReady();
+		Handle.MarkUnready();
+		REQUIRE(Volumes.Num() == 1);
+	}
+
+	SECTION("Volume is marked ready twice")
+	{
+		FSpatialReadinessHandle Handle = SpatialReadiness->CreateHandle(TestBox, TestDesc);
+		Handle.MarkReady();
+		Handle.MarkReady();
+	}
+
+	SECTION("Volume is marked unready when already unready")
+	{
+		FSpatialReadinessHandle Handle = SpatialReadiness->CreateHandle(TestBox, TestDesc);
+		Handle.MarkUnready();
+	}
+
+	SECTION("Volume is removed when handle goes out of scope")
+	{
+		{
+			FSpatialReadinessHandle Handle = SpatialReadiness->CreateHandle(TestBox, TestDesc);
+			REQUIRE(Volumes.Num() == 1);
+		}
+
+		REQUIRE(Volumes.Num() == 0);
+	}
+
+	SECTION("Handle becomes valid when api goes out of scope")
+	{
+		// Make a volume handle and require that it's valid
+		FSpatialReadinessHandle Handle = SpatialReadiness->CreateHandle(TestBox, TestDesc);
+		REQUIRE(Handle.IsValid());
+		REQUIRE(Volumes.Num() == 1);
+
+		// Delete the readiness api and ensure that the handle is now invalid
+		SpatialReadiness.Reset();
+		REQUIRE(!Handle.IsValid());
+	}
+
+	/*
+	SECTION("Ensure is fired when api callback causes its own deletion")
+	{
+		struct FSelfDeletingSpatialReadinessProvider
+		{
+			FSelfDeletingSpatialReadinessProvider() : SpatialReadiness(MakeUnique<FSpatialReadinessAPI>(this, &FSelfDeletingSpatialReadinessProvider::AddUnreadyVolume, &FSelfDeletingSpatialReadinessProvider::RemoveUnreadyVolume))
+			{ }
+
+			int32 AddUnreadyVolume(const FBox&, const FString&) { return 0; }
+			void RemoveUnreadyVolume(const int32)
+			{
+				// This is a terrible thing to do, but the idea is to 
+				SpatialReadiness.Reset();
+			}
+
+			TUniquePtr<FSpatialReadinessAPI> SpatialReadiness;
+		};
+
+		// Create an instance of the provider, generate a handle, and mark
+		// it ready to delete it. This should trigger a warning.
+		FSelfDeletingSpatialReadinessProvider Provider;
+		FSpatialReadinessHandle Handle = Provider.SpatialReadiness->CreateHandle(TestBox, TestDesc);
+		Handle.MarkReady();
+	}
+	*/
+}
+
+TEST_CASE("Spatial readiness provider class", "[spatial readiness]")
+{
+	// Create a spatial readiness provider class to mimic ones that might exist
+	struct FTestSpatialReadinessProvider
+	{
+		FTestSpatialReadinessProvider() : SpatialReadiness(this, &FTestSpatialReadinessProvider::AddUnreadyVolume, &FTestSpatialReadinessProvider::RemoveUnreadyVolume)
+		{ }
+
+		int32 AddUnreadyVolume(const FBox&, const FString&) { ++VolumesAdded; return 0; }
+
+		void RemoveUnreadyVolume(const int32) { ++VolumesRemoved; }
+
+		FSpatialReadinessAPI SpatialReadiness;
+
+		int32 VolumesAdded = 0;
+		int32 VolumesRemoved = 0;
+	};
+
+	// Make an instance of the provider
+	FTestSpatialReadinessProvider Provider;
+
+	SECTION("Instantiate spatial readiness api with member function bindings")
+	{
+		{
+			// This should trigger calls to add unready volume
+			FSpatialReadinessHandle Handle = Provider.SpatialReadiness.CreateHandle(TestBox, TestDesc);
+			REQUIRE(Handle.IsValid());
+			REQUIRE(Handle.IsReady() == false);
+			REQUIRE(Provider.VolumesAdded == 1);
+			REQUIRE(Provider.VolumesRemoved == 0);
+		}
+
+		// The handle going out of scope should cause the volume to be removed
+		REQUIRE(Provider.VolumesAdded == 1);
+		REQUIRE(Provider.VolumesRemoved == 1);
 	}
 }
\ No newline at end of file
