Add tests for adding and removing tons of volumes. Fix bug in HashMappedArray where hash table was being iterated over while being modified

diff --git a/Engine/Plugins/Experimental/SpatialReadiness/Source/SpatialReadiness/Private/SpatialReadiness/SpatialReadinessVolume.cpp b/Engine/Plugins/Experimental/SpatialReadiness/Source/SpatialReadiness/Private/SpatialReadiness/SpatialReadinessVolume.cpp
--- a/Engine/Plugins/Experimental/SpatialReadiness/Source/SpatialReadiness/Private/SpatialReadiness/SpatialReadinessVolume.cpp
+++ b/Engine/Plugins/Experimental/SpatialReadiness/Source/SpatialReadiness/Private/SpatialReadiness/SpatialReadinessVolume.cpp
@@ -149,3 +149,10 @@ const FString& FSpatialReadinessVolume::GetDescription() const
 	return EmptyString;
 #endif
 }
+
+#if !UE_BUILD_SHIPPING 
+const FString FSpatialReadinessVolume::GetDebugString() const
+{
+	return FString::Printf(TEXT("Spatial Readiness Volume [Description: %s][Bounds: %s : %s][Index: %d]"), *GetDescription(), *Bounds.Min.ToString(), *Bounds.Max.ToString(), Index);
+}
+#endif
diff --git a/Engine/Plugins/Experimental/SpatialReadiness/Source/SpatialReadiness/Public/SpatialReadinessVolume.h b/Engine/Plugins/Experimental/SpatialReadiness/Source/SpatialReadiness/Public/SpatialReadinessVolume.h
--- a/Engine/Plugins/Experimental/SpatialReadiness/Source/SpatialReadiness/Public/SpatialReadinessVolume.h
+++ b/Engine/Plugins/Experimental/SpatialReadiness/Source/SpatialReadiness/Public/SpatialReadinessVolume.h
@@ -39,6 +39,11 @@ public:
 	// the WITH_SPATIAL_READINESS_DESCRIPTIONS macro is 0
 	SPATIALREADINESS_API const FString& GetDescription() const;
 
+#if !UE_BUILD_SHIPPING 
+	// Get a debug string for this volume
+	SPATIALREADINESS_API const FString GetDebugString() const;
+#endif
+
 	// When a handle is destroyed, it will make the corresponding volume as ready
 	SPATIALREADINESS_API ~FSpatialReadinessVolume();
 
diff --git a/Engine/Plugins/Experimental/SpatialReadiness/Tests/Private/PhysicsSpatialReadinessTests.cpp b/Engine/Plugins/Experimental/SpatialReadiness/Tests/Private/PhysicsSpatialReadinessTests.cpp
--- a/Engine/Plugins/Experimental/SpatialReadiness/Tests/Private/PhysicsSpatialReadinessTests.cpp
+++ b/Engine/Plugins/Experimental/SpatialReadiness/Tests/Private/PhysicsSpatialReadinessTests.cpp
@@ -208,4 +208,72 @@ TEST_CASE("Sim Callback", "[physics spatial readiness]")
 		REQUIRE(bIsReady);
 		REQUIRE(VolumeIndices.Num() == 0);
 	}
+
+	SECTION("Add and remove volumes in different orders")
+	{
+		FBox UnreadyBox(FVector(-1), FVector(1));
+		TArray<FSpatialReadinessVolume> Volumes;
+
+		// add/remove
+		Volumes.Emplace(SpatialReadiness.CreateVolume(UnreadyBox, TEXT("Test Volume")));
+		AdvanceAndWait(Solver);
+		Volumes.RemoveAt(0);
+		AdvanceAndWait(Solver);
+
+		// add/add/remove
+		Volumes.Emplace(SpatialReadiness.CreateVolume(UnreadyBox, TEXT("Test Volume")));
+		AdvanceAndWait(Solver);
+		Volumes.Emplace(SpatialReadiness.CreateVolume(UnreadyBox, TEXT("Test Volume")));
+		AdvanceAndWait(Solver);
+		Volumes.RemoveAt(0);
+		AdvanceAndWait(Solver);
+	}
+
+	SECTION("Add and remove many random volumes")
+	{
+		const int32 NumVolumeActions = 100;
+		FRandomStream RandStream(42);
+
+		TArray<FSpatialReadinessVolume> Volumes;
+
+		for (int32 Index = 0; Index < NumVolumeActions; ++Index)
+		{
+			const int32 ActionIndex = RandStream.RandRange(0, 1);
+			switch (ActionIndex)
+			{
+				// Add a volume
+				case 0:
+				{
+					const FVector BoxMin = FVector(
+						RandStream.FRandRange(-2000, 1000),
+						RandStream.FRandRange(-2000, 1000),
+						RandStream.FRandRange(-2000, 1000));
+					const FVector BoxMax = BoxMin + FVector(
+						RandStream.FRandRange(0, 1000),
+						RandStream.FRandRange(0, 1000),
+						RandStream.FRandRange(0, 1000));
+					const FBox UnreadyBox(BoxMin, BoxMax);
+					Volumes.Emplace(SpatialReadiness.CreateVolume(UnreadyBox, TEXT("Test Volume")));
+
+					WARN(*FString::Printf(TEXT("Adding: %s"), *Volumes[Volumes.Num()-1].GetDebugString()));
+
+					break;
+				}
+
+				// Remove a volume
+				case 1:
+				{
+					if (Volumes.Num() == 0) { break; }
+					const int32 VolumeIndex = RandStream.RandRange(0, Volumes.Num() - 1);
+
+					WARN(*FString::Printf(TEXT("Removing: %s"), * Volumes[VolumeIndex].GetDebugString()));
+
+					Volumes.RemoveAt(VolumeIndex);
+					break;
+				}
+			}
+
+			AdvanceAndWait(Solver);
+		}
+	}
 }
diff --git a/Engine/Source/Runtime/Experimental/Chaos/Public/Chaos/Framework/HashMappedArray.h b/Engine/Source/Runtime/Experimental/Chaos/Public/Chaos/Framework/HashMappedArray.h
--- a/Engine/Source/Runtime/Experimental/Chaos/Public/Chaos/Framework/HashMappedArray.h
+++ b/Engine/Source/Runtime/Experimental/Chaos/Public/Chaos/Framework/HashMappedArray.h
@@ -154,6 +154,8 @@ namespace Chaos::Private
 						// Re-add the last element's key, reusing the index of the element we just removed
 						HashTable.Add(LastElemKey, Index);
 					}
+
+					break;
 				}
 			}
 		}
