Remove debug WARNings from spatial readiness tests

diff --git a/Engine/Plugins/Experimental/SpatialReadiness/Tests/Private/PhysicsSpatialReadinessTests.cpp b/Engine/Plugins/Experimental/SpatialReadiness/Tests/Private/PhysicsSpatialReadinessTests.cpp
--- a/Engine/Plugins/Experimental/SpatialReadiness/Tests/Private/PhysicsSpatialReadinessTests.cpp
+++ b/Engine/Plugins/Experimental/SpatialReadiness/Tests/Private/PhysicsSpatialReadinessTests.cpp
@@ -254,9 +254,6 @@ TEST_CASE("Sim Callback", "[physics spatial readiness]")
 						RandStream.FRandRange(0, 1000));
 					const FBox UnreadyBox(BoxMin, BoxMax);
 					Volumes.Emplace(SpatialReadiness.CreateVolume(UnreadyBox, TEXT("Test Volume")));
-
-					WARN(*FString::Printf(TEXT("Adding: %s"), *Volumes[Volumes.Num()-1].GetDebugString()));
-
 					break;
 				}
 
@@ -265,9 +262,6 @@ TEST_CASE("Sim Callback", "[physics spatial readiness]")
 				{
 					if (Volumes.Num() == 0) { break; }
 					const int32 VolumeIndex = RandStream.RandRange(0, Volumes.Num() - 1);
-
-					WARN(*FString::Printf(TEXT("Removing: %s"), * Volumes[VolumeIndex].GetDebugString()));
-
 					Volumes.RemoveAt(VolumeIndex);
 					break;
 				}
