Fix for state graph unit tests.
#rnx

diff --git a/Engine/Plugins/Experimental/StateGraph/Tests/StateGraphTests/Private/StateGraphTestsUnit.cpp b/Engine/Plugins/Experimental/StateGraph/Tests/StateGraphTests/Private/StateGraphTestsUnit.cpp
--- a/Engine/Plugins/Experimental/StateGraph/Tests/StateGraphTests/Private/StateGraphTestsUnit.cpp
+++ b/Engine/Plugins/Experimental/StateGraph/Tests/StateGraphTests/Private/StateGraphTestsUnit.cpp
@@ -110,8 +110,9 @@ TEST_CASE("FStateGraph ContextName", "[FStateGraphContextName]")
 	FStateGraphRef StateGraph(MakeShared<FStateGraph>("Test", TEXT("TestContextName")));
 	FStateGraphNodeFunctionRef TestNodeA = StateGraph->CreateNode("TestNodeA", [](FStateGraph& StateGraph, FStateGraphNodeFunctionComplete Complete) { Complete(); });
 	StateGraph->Run();
-	CHECK(StateGraph->GetContextName() == TEXT("Test(TestContextName)"));
-	CHECK(TestNodeA->GetContextName() == TEXT("Test(TestContextName).TestNodeA"));
+	CHECK(StateGraph->GetContextName() == TEXT("TestContextName"));
+	CHECK(StateGraph->GetLogName() == TEXT("Test(TestContextName)"));
+	CHECK(TestNodeA->GetLogName() == TEXT("Test(TestContextName).TestNodeA"));
 }
 
 class FClassNode : public FStateGraphNode
