Freezes due to midphase detection and new particle spawns tracked in separate lists to avoid ordering issues

diff --git a/Engine/Plugins/Experimental/SpatialReadiness/Source/SpatialReadiness/Private/SpatialReadiness/SpatialReadinessSimCallback.cpp b/Engine/Plugins/Experimental/SpatialReadiness/Source/SpatialReadiness/Private/SpatialReadiness/SpatialReadinessSimCallback.cpp
--- a/Engine/Plugins/Experimental/SpatialReadiness/Source/SpatialReadiness/Private/SpatialReadiness/SpatialReadinessSimCallback.cpp
+++ b/Engine/Plugins/Experimental/SpatialReadiness/Source/SpatialReadiness/Private/SpatialReadiness/SpatialReadinessSimCallback.cpp
@@ -386,12 +386,13 @@ void FSpatialReadinessSimCallback::FreezeParticles_PT()
 		return;
 	}
 
-	for (FPBDRigidParticleHandle* RigidParticle : UnreadyRigidParticles_PT)
+	ForEachUnreadyRigidParticle_PT([&](Chaos::FPBDRigidParticleHandle* RigidParticle)
 	{
 		const EObjectStateType ObjectState = RigidParticle->ObjectState();
 		if (ObjectState == EObjectStateType::Static)
 		{
-			return;
+			// This object is already static, so continue on to the next one
+			return true;
 		}
 
 		// Get the geometry particle
@@ -420,7 +421,10 @@ void FSpatialReadinessSimCallback::FreezeParticles_PT()
 			CVD_TRACE_DEBUG_DRAW_BOX(FBox(Bounds.Min(), Bounds.Max()), Tag, FColor::Yellow, FChaosVisualDebuggerTrace::GetSolverID(*GetSolver()));
 		}
 #endif
-	}
+
+		// Continue on to the next one
+		return true;
+	});
 }
 
 void FSpatialReadinessSimCallback::UnFreezeParticles_PT()
@@ -475,6 +479,12 @@ void FSpatialReadinessSimCallback::OnParticlesRegistered_Internal(TArray<FSingle
 {
 	SCOPE_CYCLE_COUNTER(STAT_SpatialReadiness_Physics_ParticlesRegistered)
 
+	// Empty the list of newly registered particles which are unready,
+	// so that we don't double count them
+	//Registered_UnreadyRigidParticles_PT.Reset();
+	Swap(NewRegistered_UnreadyRigidParticles_PT, OldRegistered_UnreadyRigidParticles_PT);
+	NewRegistered_UnreadyRigidParticles_PT.Reset();
+
 	// For each newly added particle, query against unready volumes to see if
 	// it should be frozen.
 	for (FSingleParticlePhysicsProxy* ParticleProxy : RegisteredProxies)
@@ -498,7 +508,7 @@ void FSpatialReadinessSimCallback::OnParticlesRegistered_Internal(TArray<FSingle
 		{
 			// If this volume hit an unready volume, query readiness returns false
 			// and we must mark the particle as frozen
-			UnreadyRigidParticles_PT.Add(RigidParticle);
+			NewRegistered_UnreadyRigidParticles_PT.Add(RigidParticle);
 		}
 	}
 }
@@ -508,7 +518,7 @@ void FSpatialReadinessSimCallback::OnMidPhaseModification_Internal(FMidPhaseModi
 	SCOPE_CYCLE_COUNTER(STAT_SpatialReadiness_Physics_MidPhase)
 
 	// Clear the list of unready particles
-	UnreadyRigidParticles_PT.Reset();
+	MidPhase_UnreadyRigidParticles_PT.Reset();
 
 	// Go through every unready volume
 	for (FSingleParticlePhysicsProxy* UnreadyProxy : UnreadyVolumeParticles_PT)
@@ -567,7 +577,7 @@ void FSpatialReadinessSimCallback::OnMidPhaseModification_Internal(FMidPhaseModi
 #endif
 
 			// Add the particle to the list of particles to freeze
-			UnreadyRigidParticles_PT.Add(RigidParticle);
+			MidPhase_UnreadyRigidParticles_PT.Add(RigidParticle);
 
 			// Disable the midphase
 			MidPhase.Disable();
@@ -594,3 +604,30 @@ uint32 FSpatialReadinessSimCallback::FHashMapTraits::GetElementID(const FUnready
 	return INDEX_NONE;
 }
 
+void FSpatialReadinessSimCallback::ForEachUnreadyRigidParticle_PT(const TFunction<bool(Chaos::FPBDRigidParticleHandle*)>& Lambda) const
+{
+	for (Chaos::FPBDRigidParticleHandle* Particle : NewRegistered_UnreadyRigidParticles_PT)
+	{
+		if (Lambda(Particle) == false)
+		{
+			return;
+		}
+	}
+
+	for (Chaos::FPBDRigidParticleHandle* Particle : OldRegistered_UnreadyRigidParticles_PT)
+	{
+		if (Lambda(Particle) == false)
+		{
+			return;
+		}
+	}
+
+	for (Chaos::FPBDRigidParticleHandle* Particle : MidPhase_UnreadyRigidParticles_PT)
+	{
+		if (Lambda(Particle) == false)
+		{
+			return;
+		}
+	}
+}
+
diff --git a/Engine/Plugins/Experimental/SpatialReadiness/Source/SpatialReadiness/Public/SpatialReadinessSimCallback.h b/Engine/Plugins/Experimental/SpatialReadiness/Source/SpatialReadiness/Public/SpatialReadinessSimCallback.h
--- a/Engine/Plugins/Experimental/SpatialReadiness/Source/SpatialReadiness/Public/SpatialReadinessSimCallback.h
+++ b/Engine/Plugins/Experimental/SpatialReadiness/Source/SpatialReadiness/Public/SpatialReadinessSimCallback.h
@@ -118,8 +118,28 @@ protected:
 	// List of particle handles which represent unready volumes
 	TSet<Chaos::FSingleParticlePhysicsProxy*> UnreadyVolumeParticles_PT;
 
-	// List of particle handles which represent particles that interacted with unready volumes
-	TSet<Chaos::FPBDRigidParticleHandle*> UnreadyRigidParticles_PT;
+	// Lists of particle handles which represent particles that interacted with unready volumes
+	//
+	// This got a little bit confusing due to some quirks of the particle
+	// creation pipeline.
+	//
+	// A particle is added to the NewRegistered list if it's initially
+	// overlapping an unready volume. Ideally it would then get picked
+	// up in the unready list maintained by the mid-phase modifier.
+	// However, it isn't inserted into the acceleration structure
+	// immediately, so it is missed.
+	//
+	// Instead, at the beginning of the next frame, NewRegistered contents
+	// are moved into OldRegistered, and NewRegistered is emptied.
+	//
+	// This ensures that particles which are spawned on top of unready
+	// volumes will not be able to sneak in a integration step before
+	// they have been inserted into the acceleration structure, where the
+	// mid-phase modifier will be able to manage them.
+	TSet<Chaos::FPBDRigidParticleHandle*> NewRegistered_UnreadyRigidParticles_PT;
+	TSet<Chaos::FPBDRigidParticleHandle*> OldRegistered_UnreadyRigidParticles_PT;
+	TSet<Chaos::FPBDRigidParticleHandle*> MidPhase_UnreadyRigidParticles_PT;
+	void ForEachUnreadyRigidParticle_PT(const TFunction<bool(Chaos::FPBDRigidParticleHandle*)>& Lambda) const;
 
 	// "Unready" particles are forced to be stationary in PreSimulate, and restored
 	// to their previous state in PostIntegrate. The values needed for restoration
diff --git a/Engine/Plugins/Experimental/SpatialReadiness/Tests/Private/PhysicsSpatialReadinessTests.cpp b/Engine/Plugins/Experimental/SpatialReadiness/Tests/Private/PhysicsSpatialReadinessTests.cpp
--- a/Engine/Plugins/Experimental/SpatialReadiness/Tests/Private/PhysicsSpatialReadinessTests.cpp
+++ b/Engine/Plugins/Experimental/SpatialReadiness/Tests/Private/PhysicsSpatialReadinessTests.cpp
@@ -10,6 +10,18 @@
 #include "PhysicsProxy/SingleParticlePhysicsProxy.h"
 #include "Physics/PhysicsFiltering.h"
 
+using Chaos::FSingleParticlePhysicsProxy;
+using Chaos::FPBDRigidParticleHandle;
+using Chaos::FVec3;
+using Chaos::FReal;
+using Chaos::FRigidBodyHandle_External;
+using Chaos::TBox;
+using Chaos::FImplicitObjectPtr;
+using Chaos::EObjectStateType;
+using Chaos::FPBDRigidsSolver;
+using Chaos::FPhysicsThreadContextScope;
+using Chaos::EThreadingModeTemp;
+
 namespace
 {
 	const FBox UnitBounds(FVector(-.5), FVector(.5));
@@ -28,7 +40,7 @@ namespace
 		// Create a box implicit geometry from the same bounds as the unready volume
 		const FVec3 BoxCenter = (Bounds.Min + Bounds.Max) * .5f;
 		const FVec3 BoxHalfExtent = Bounds.Max - Bounds.Min;
-		Chaos::FImplicitObjectPtr BoxGeom = MakeImplicitObjectPtr<TBox<FReal, 3>>(-BoxHalfExtent, BoxHalfExtent);
+		FImplicitObjectPtr BoxGeom = MakeImplicitObjectPtr<TBox<FReal, 3>>(-BoxHalfExtent, BoxHalfExtent);
 
 		// Create a new static particle to represent the volume
 		FActorCreationParams Params;
@@ -89,9 +101,21 @@ namespace
 			return UnreadyVolumeParticles_PT;
 		}
 
-		const TSet<FPBDRigidParticleHandle*>& GetUnreadyRigidParticles_PT() const
+		const TSet<FPBDRigidParticleHandle*> GetUnreadyRigidParticles_PT() const
+		{
+			TSet<FPBDRigidParticleHandle*> RigidParticles;
+			ForEachUnreadyRigidParticle_PT([&](FPBDRigidParticleHandle* RigidParticle) { RigidParticles.Add(RigidParticle); return true; });
+			return RigidParticles;
+		}
+
+		TFunction<void()> OnPreSimulate_Callback = [](){};
+
+	protected:
+
+		virtual void OnPreSimulate_Internal() override
 		{
-			return UnreadyRigidParticles_PT;
+			FSpatialReadinessSimCallback::OnPreSimulate_Internal();
+			OnPreSimulate_Callback();
 		}
 	};
 }
@@ -113,6 +137,9 @@ TEST_CASE("Sim Callback", "[physics spatial readiness]")
 		&FTestSpatialReadinessSimCallback::AddUnreadyVolume_GT,
 		&FTestSpatialReadinessSimCallback::RemoveUnreadyVolume_GT);
 
+	// Make a physics thread scope to avoid any thread context checks
+	FPhysicsThreadContextScope PTScope(true);
+
 	SECTION("Spawn dynamic particle overlapping unready volume, no gravity")
 	{
 		// Make a volume
@@ -124,6 +151,9 @@ TEST_CASE("Sim Callback", "[physics spatial readiness]")
 		FSingleParticlePhysicsProxy* BoxProxy = MakeDynamicBox(Scene, UnitBounds, false);
 		AdvanceAndWait(Solver);
 
+		// Make sure the particle is frozen immediately by the registry callback
+		REQUIRE(SimCallback->GetUnreadyRigidParticles_PT().Num() == 1);
+
 		// Advance one more tick to detect the midphase
 		AdvanceAndWait(Solver);
 		REQUIRE(SimCallback->GetUnreadyRigidParticles_PT().Num() == 1);
@@ -139,17 +169,22 @@ TEST_CASE("Sim Callback", "[physics spatial readiness]")
 		// Make a particle which should overlap that volume
 		FSingleParticlePhysicsProxy* BoxProxy = MakeDynamicBox(Scene, UnitBounds, true);
 		AdvanceAndWait(Solver);
+		REQUIRE(SimCallback->GetUnreadyRigidParticles_PT().Num() == 1);
+
+		// Record the initial position
+		const FVec3 FrozenX0 = BoxProxy->GetPhysicsThreadAPI()->X();
 
 		// Advance one more tick to detect the midphase
 		AdvanceAndWait(Solver);
 		REQUIRE(SimCallback->GetUnreadyRigidParticles_PT().Num() == 1);
+
+		// Make sure the particle hasn't moved
+		const FVec3 FrozenX1 = BoxProxy->GetPhysicsThreadAPI()->X();
+		REQUIRE(FrozenX0 == FrozenX1);
 	}
 
 	SECTION("Dynamic particle with gravity stops falling when it hits an unready volume")
 	{
-		// Make a physics thread scope to avoid any thread context checks
-		FPhysicsThreadContextScope PTScope(true);
-
 		// Make a volume
 		FBox UnreadyBox(FVector(-1000, -1000, -10000), FVector(1000, 1000, 0));
 		FSpatialReadinessVolume Volume = SpatialReadiness.CreateVolume(UnreadyBox, TEXT("Test Volume"));
@@ -187,6 +222,35 @@ TEST_CASE("Sim Callback", "[physics spatial readiness]")
 		REQUIRE(FrozenX2.Z < FrozenX0.Z);
 	}
 
+	SECTION("Force added to frozen dynamic particle does nothing")
+	{
+		// Make a volume
+		FSpatialReadinessVolume Volume = SpatialReadiness.CreateVolume(UnitBounds, TEXT("Test Volume"));
+		AdvanceAndWait(Solver);
+		REQUIRE(SimCallback->GetUnreadyRigidParticles_PT().Num() == 0);
+
+		// Make a particle which should overlap that volume
+		FSingleParticlePhysicsProxy* BoxProxy = MakeDynamicBox(Scene, UnitBounds, true);
+		AdvanceAndWait(Solver);
+
+		// Get the position of the box
+		const FVec3 FrozenX0 = BoxProxy->GetPhysicsThreadAPI()->X();
+
+		// Add a force to the box in pre-simulate
+		SimCallback->OnPreSimulate_Callback = [BoxProxy]()
+		{
+			BoxProxy->GetPhysicsThreadAPI()->AddForce(FVec3(100,0,0));
+		};
+
+		// Advance again and make sure the box didn't move. Do it 10 times for good measure
+		for (int32 Iteration = 0; Iteration < 10; ++Iteration)
+		{
+			AdvanceAndWait(Solver);
+			const FVec3 FrozenX1 = BoxProxy->GetPhysicsThreadAPI()->X();
+			REQUIRE(FrozenX0 == FrozenX1);
+		}
+	}
+
 	SECTION("Game thread query for readiness")
 	{
 		// Make a volume
