Update spatial readiness tests for changes from 46780709

The change intentionally added one frame of delay un-freeze frozen rigids and this needs to be reflected in unit tests

diff --git a/Engine/Plugins/Experimental/SpatialReadiness/Tests/Private/PhysicsSpatialReadinessTests.cpp b/Engine/Plugins/Experimental/SpatialReadiness/Tests/Private/PhysicsSpatialReadinessTests.cpp
--- a/Engine/Plugins/Experimental/SpatialReadiness/Tests/Private/PhysicsSpatialReadinessTests.cpp
+++ b/Engine/Plugins/Experimental/SpatialReadiness/Tests/Private/PhysicsSpatialReadinessTests.cpp
@@ -163,13 +163,14 @@ TEST_CASE("Sim Callback", "[physics spatial readiness]")
 		// Record the initial position
 		const FVec3 FrozenX0 = BoxProxy->GetPhysicsThreadAPI()->X();
 
-		// Advance one more tick to detect the midphase
-		AdvanceAndWait(Solver);
-		REQUIRE(SimCallback->GetUnreadyRigidParticles_PT().Num() == 1);
-
-		// Make sure the particle hasn't moved
-		const FVec3 FrozenX1 = BoxProxy->GetPhysicsThreadAPI()->X();
-		REQUIRE(FrozenX0 == FrozenX1);
+		// Advance ten more ticks to make sure it never starts falling
+		for (int32 Iteration = 0; Iteration < 10; ++Iteration)
+		{
+			AdvanceAndWait(Solver);
+			REQUIRE(SimCallback->GetUnreadyRigidParticles_PT().Num() == 1);
+			const FVec3 FrozenX1 = BoxProxy->GetPhysicsThreadAPI()->X();
+			REQUIRE(FrozenX0 == FrozenX1);
+		}
 	}
 
 	SECTION("Dynamic particle with gravity stops falling when it hits an unready volume")
@@ -180,7 +181,7 @@ TEST_CASE("Sim Callback", "[physics spatial readiness]")
 		AdvanceAndWait(Solver);
 		REQUIRE(SimCallback->GetUnreadyRigidParticles_PT().Num() == 0);
 
-		// Make a particle which should overlap that volume
+		// Make a particle which is just above the unready volume
 		FBox FallingBox(FVector(-1, -1, 2), FVector(1, 1, 3));
 		FSingleParticlePhysicsProxy* BoxProxy = MakeDynamicBox(Scene, FallingBox, true);
 		AdvanceAndWait(Solver);
@@ -199,9 +200,12 @@ TEST_CASE("Sim Callback", "[physics spatial readiness]")
 		}
 
 		// Mark the volume as "ready" and advance again - the particle
-		// should be removed from the list of unready particles
+		// should be removed from the list of unready particles after 2 ticks
+		// because spatial readiness always keeps particles unready for one extra tick.
 		Volume.MarkReady();
 		AdvanceAndWait(Solver);
+		REQUIRE(SimCallback->GetUnreadyRigidParticles_PT().Num() == 1);
+		AdvanceAndWait(Solver);
 		REQUIRE(SimCallback->GetUnreadyRigidParticles_PT().Num() == 0);
 
 		// Advancing one more frame, we should see the particle start to fall
