Update SR unit test to reflect the 2 frame counter before unfreeze

diff --git a/Engine/Plugins/Experimental/SpatialReadiness/Tests/Private/PhysicsSpatialReadinessTests.cpp b/Engine/Plugins/Experimental/SpatialReadiness/Tests/Private/PhysicsSpatialReadinessTests.cpp
--- a/Engine/Plugins/Experimental/SpatialReadiness/Tests/Private/PhysicsSpatialReadinessTests.cpp
+++ b/Engine/Plugins/Experimental/SpatialReadiness/Tests/Private/PhysicsSpatialReadinessTests.cpp
@@ -206,6 +206,8 @@ TEST_CASE("Sim Callback", "[physics spatial readiness]")
 		AdvanceAndWait(Solver);
 		REQUIRE(SimCallback->GetUnreadyRigidParticles_PT().Num() == 1);
 		AdvanceAndWait(Solver);
+		REQUIRE(SimCallback->GetUnreadyRigidParticles_PT().Num() == 1);
+		AdvanceAndWait(Solver);
 		REQUIRE(SimCallback->GetUnreadyRigidParticles_PT().Num() == 0);
 
 		// Advancing one more frame, we should see the particle start to fall
