Fixing spatial readiness tests failing to compile due to a missing test target rule.

diff --git a/Engine/Plugins/Experimental/SpatialReadiness/Tests/SpatialReadinessTests.Target.cs b/Engine/Plugins/Experimental/SpatialReadiness/Tests/SpatialReadinessTests.Target.cs
--- a/Engine/Plugins/Experimental/SpatialReadiness/Tests/SpatialReadinessTests.Target.cs
+++ b/Engine/Plugins/Experimental/SpatialReadiness/Tests/SpatialReadinessTests.Target.cs
@@ -9,5 +9,6 @@ public class SpatialReadinessTestsTarget : TestTargetRules
 	{
 		bUsePlatformFileStub = true;
 		bMockEngineDefaults = true;
+		TestTargetRules.bTestsRequireEngine = true;
 	}
 }
