Chaos API - add scene settings

#Synced-CL 48771487

diff --git a/Engine/Plugins/Experimental/ChaosRigidPhysicsAsync/Source/ChaosRigidPhysicsAsync/Private/ChaosRigidPhysicsAsync/RigidFactoryAsync.cpp b/Engine/Plugins/Experimental/ChaosRigidPhysicsAsync/Source/ChaosRigidPhysicsAsync/Private/ChaosRigidPhysicsAsync/RigidFactoryAsync.cpp
--- a/Engine/Plugins/Experimental/ChaosRigidPhysicsAsync/Source/ChaosRigidPhysicsAsync/Private/ChaosRigidPhysicsAsync/RigidFactoryAsync.cpp
+++ b/Engine/Plugins/Experimental/ChaosRigidPhysicsAsync/Source/ChaosRigidPhysicsAsync/Private/ChaosRigidPhysicsAsync/RigidFactoryAsync.cpp
@@ -1,10 +1,11 @@
 // Copyright Epic Games, Inc. All Rights Reserved.
 #include "ChaosRigidPhysicsAsync/RigidFactoryAsync.h"
 
-#include "ChaosRigidPhysicsAsync/RigidBodyAsync.h"
-#include "ChaosRigidPhysicsAsync/RigidBodyContainerAsync.h"
 #include "ChaosRigidPhysicsAsync/RigidSceneAsync.h"
+#include "ChaosRigidPhysicsAsync/RigidSceneSettingsAsync.h"
+#include "Logging/StructuredLog.h"
 #include "RigidPhysics/RigidBody.h"
+#include "RigidPhysics/RigidLog.h"
 
 #if UE_RIGIDPHYSICS_API_ENABLED
 
@@ -12,28 +13,45 @@ namespace Chaos::Rigids::Async
 {
 	FRigidFactoryAsync::FRigidFactoryAsync()
 	{
-
 	}
 
 	FRigidFactoryAsync::~FRigidFactoryAsync()
 	{
 		// Make sure all scenes were destroyed
-		UE_RIGIDPHYSICS_CHECK(Scenes.IsEmpty());
+		check(Scenes.IsEmpty());
+	}
+
+	const UE::Physics::FRigidTypeId& FRigidFactoryAsync::GetSceneTypeId() const
+	{
+		return FRigidSceneAsyncGT::GetStaticTypeId();
 	}
 
-	UE::Physics::FRigidSceneHandle FRigidFactoryAsync::CreateScene(const UE::Physics::FRigidDebugName& InName)
+	UE::Physics::FRigidSceneHandle FRigidFactoryAsync::CreateScene(const UE::Physics::FRigidDebugName& InName, const UE::Physics::IRigidSceneSettings* InSettings)
 	{
 		using namespace UE::Physics;
+		check(InSettings != nullptr);
+
+		const FRigidSceneSettingsAsync* SettingsAsync = InSettings->AsA<FRigidSceneSettingsAsync>();
 
-		FRigidSceneSettingsAsync Settings;
+		UE_CLOGFMT((SettingsAsync == nullptr), LogRigidPhysics, Error,
+			"FRigidFactoryAsync::CreateScene: Cannot create scene. Expected settings of type '{0}' but received '{1}'.",
+			FRigidSceneSettingsAsync::GetStaticTypeId().GetTypeName().GetCharArray().GetData(),
+			InSettings->GetTypeId().GetTypeName().GetCharArray().GetData());
 
 		// Create a scene and store a hard reference to prevent its destruction
-		FRigidSceneAsyncGT* Scene = new FRigidSceneAsyncGT(InName, Settings);
-		Scenes.Add(Scene);
+		if (SettingsAsync != nullptr)
+		{
+			if (FRigidSceneAsyncGT* Scene = new FRigidSceneAsyncGT(InName, *SettingsAsync))
+			{
+				Scenes.Add(Scene);
 
-		Scene->Startup();
+				Scene->Startup();
+
+				return FRigidSceneHandle(Scene);
+			}
+		}
 
-		return FRigidSceneHandle(Scene);
+		return FRigidSceneHandle();
 	}
 
 	UE::Physics::IRigidScene* FRigidFactoryAsync::CreateSceneFromSolver(FPBDRigidsSolver* InSolver)
diff --git a/Engine/Plugins/Experimental/ChaosRigidPhysicsAsync/Source/ChaosRigidPhysicsAsync/Private/ChaosRigidPhysicsAsync/RigidSceneAsync.cpp b/Engine/Plugins/Experimental/ChaosRigidPhysicsAsync/Source/ChaosRigidPhysicsAsync/Private/ChaosRigidPhysicsAsync/RigidSceneAsync.cpp
--- a/Engine/Plugins/Experimental/ChaosRigidPhysicsAsync/Source/ChaosRigidPhysicsAsync/Private/ChaosRigidPhysicsAsync/RigidSceneAsync.cpp
+++ b/Engine/Plugins/Experimental/ChaosRigidPhysicsAsync/Source/ChaosRigidPhysicsAsync/Private/ChaosRigidPhysicsAsync/RigidSceneAsync.cpp
@@ -7,6 +7,8 @@
 
 namespace Chaos::Rigids::Async
 {
+	UE_RIGIDPHYSICS_RIGIDTYPED_IMPL(FRigidSceneAsyncGT, UE::Physics::FRigidScene);
+
 	FRigidSceneAsyncGT::FRigidSceneAsyncGT(
 		const UE::Physics::FRigidDebugName& InName,
 		const FRigidSceneSettingsAsync& InSceneSettings)
@@ -76,9 +78,14 @@ namespace Chaos::Rigids::Async
 
 		if (Solver == nullptr)
 		{
+			// TODO_CHAOSAPI: We should contain an FChaosScene here rather than a raw FPBDRigidsSolver because
+			// FChaosScene holds the spatail acceleration. However FPhysicsScene_Chaos inherits FChaosScene so
+			// we need to break that inheritance first. Or more likely, extract the context of FChaosScene into
+			// a new class which can be used here and there.
+
 			// NOTE: Not using full Async for now - it works but it complicates testing (time shift and interpolation)
 			FChaosSolversModule* Module = FChaosSolversModule::GetModule();
-			Solver = Module->CreateSolver(/*Owner=*/nullptr, /*AsyncDt*/-1, EThreadingMode::TaskGraph, FName(*Name.Value()));
+			Solver = Module->CreateSolver(/*Owner=*/nullptr, SceneSettings.AsyncDt, SceneSettings.ThreadingMode, FName(*Name.Value()));
 			bOwnsSolver = true;
 		}
 
@@ -210,13 +217,17 @@ namespace Chaos::Rigids::Async
 		using namespace UE::Physics;
 		CheckGameWriteLocked();
 
-		Solver->EnqueueCommandImmediate(
-			[this, BodyHandle = InBody->GetHandle()]()
-			{
-				ScenePT->DestroyBodyPT(BodyHandle);
-			});
+		// TODO_CHAOSAPI: Handle destruction of bodies from other containers
+		if (FRigidBodyAsyncGT* Body = InBody->AsA<FRigidBodyAsyncGT>())
+		{
+			Solver->EnqueueCommandImmediate(
+				[this, BodyHandle = Body->GetHandle()]()
+				{
+					ScenePT->DestroyBodyPT(BodyHandle);
+				});
 
-		BodyPool->DestroyBody(InBody);
+			BodyPool->DestroyBody(Body);
+		}
 	}
 
 	void FRigidSceneAsyncGT::RegisterModifier(UE::Physics::IRigidSceneModifier* Modifier)
@@ -356,6 +367,8 @@ namespace Chaos::Rigids::Async
 
 namespace Chaos::Rigids::Async
 {
+	UE_RIGIDPHYSICS_RIGIDTYPED_IMPL(FRigidSceneAsyncPT, UE::Physics::IRigidScene);
+
 	FRigidSceneAsyncPT::FRigidSceneAsyncPT(const UE::Physics::FRigidDebugName& InName, const UE::Physics::FRigidBodyContainerHandle& InBodyPoolHandle)
 		: Name(InName)
 		, Solver(nullptr)
diff --git a/Engine/Plugins/Experimental/ChaosRigidPhysicsAsync/Source/ChaosRigidPhysicsAsync/Private/ChaosRigidPhysicsAsync/RigidSceneSettingsAsync.cpp b/Engine/Plugins/Experimental/ChaosRigidPhysicsAsync/Source/ChaosRigidPhysicsAsync/Private/ChaosRigidPhysicsAsync/RigidSceneSettingsAsync.cpp
new file mode 100644
--- /dev/null
+++ b/Engine/Plugins/Experimental/ChaosRigidPhysicsAsync/Source/ChaosRigidPhysicsAsync/Private/ChaosRigidPhysicsAsync/RigidSceneSettingsAsync.cpp
@@ -0,0 +1,11 @@
+// Copyright Epic Games, Inc. All Rights Reserved.
+#include "ChaosRigidPhysicsAsync/RigidSceneSettingsAsync.h"
+
+#if UE_RIGIDPHYSICS_API_ENABLED
+
+namespace Chaos::Rigids::Async
+{
+	UE_RIGIDPHYSICS_RIGIDTYPED_IMPL(FRigidSceneSettingsAsync, UE::Physics::IRigidSceneSettings);
+}
+
+#endif
diff --git a/Engine/Plugins/Experimental/ChaosRigidPhysicsAsync/Source/ChaosRigidPhysicsAsync/Public/ChaosRigidPhysicsAsync/RigidFactoryAsync.h b/Engine/Plugins/Experimental/ChaosRigidPhysicsAsync/Source/ChaosRigidPhysicsAsync/Public/ChaosRigidPhysicsAsync/RigidFactoryAsync.h
--- a/Engine/Plugins/Experimental/ChaosRigidPhysicsAsync/Source/ChaosRigidPhysicsAsync/Public/ChaosRigidPhysicsAsync/RigidFactoryAsync.h
+++ b/Engine/Plugins/Experimental/ChaosRigidPhysicsAsync/Source/ChaosRigidPhysicsAsync/Public/ChaosRigidPhysicsAsync/RigidFactoryAsync.h
@@ -5,32 +5,22 @@
 
 #if UE_RIGIDPHYSICS_API_ENABLED
 
-#include "RigidPhysics/RigidScene.h"
-#include "Templates/SharedPointer.h"
-#include "Templates/UniquePtr.h"
-
-namespace Chaos
-{
-	class FPBDRigidsSolver;
-	class FSingleParticlePhysicsProxy;
-}
+#include "RigidPhysics/RigidFactory.h"
+#include "RigidPhysics/RigidSceneSettings.h"
 
 namespace Chaos::Rigids::Async
 {
-	class FRigidSceneSettingsAsync
-	{
-	public:
-		float AsyncDt = (1.0f / 60.0f);
-	};
-
-	class FRigidFactoryAsync : public UE::Physics::IRigidSceneFactory
+	// TODO_CHAOSAPI: We should be regsitering an instance of FRigidFactoryAsync with
+	// some kind of abstract factory system
+	//
+	class FRigidFactoryAsync : public UE::Physics::IRigidFactory
 	{
 	public:
 		FRigidFactoryAsync();
 		virtual ~FRigidFactoryAsync();
-
-		UE_INTERNAL CHAOSRIGIDPHYSICSASYNC_API virtual UE::Physics::FRigidSceneHandle CreateScene(const UE::Physics::FRigidDebugName& InName) override final;
-		UE_INTERNAL CHAOSRIGIDPHYSICSASYNC_API virtual void DestroyScene(const UE::Physics::FRigidSceneHandle& InSceneHandle) override final;
+		UE_INTERNAL virtual const UE::Physics::FRigidTypeId& GetSceneTypeId() const override final;
+		UE_INTERNAL UE::Physics::FRigidSceneHandle CreateScene(const UE::Physics::FRigidDebugName& InName, const UE::Physics::IRigidSceneSettings* InSettings);
+		UE_INTERNAL virtual void DestroyScene(const UE::Physics::FRigidSceneHandle& InSceneHandle) override final;
 
 		// TODO_CHAOSAPI: Temp until solver is created through new API
 		UE_INTERNAL CHAOSRIGIDPHYSICSASYNC_API UE::Physics::IRigidScene* CreateSceneFromSolver(Chaos::FPBDRigidsSolver* InSolver);
diff --git a/Engine/Plugins/Experimental/ChaosRigidPhysicsAsync/Source/ChaosRigidPhysicsAsync/Public/ChaosRigidPhysicsAsync/RigidFwdAsync.h b/Engine/Plugins/Experimental/ChaosRigidPhysicsAsync/Source/ChaosRigidPhysicsAsync/Public/ChaosRigidPhysicsAsync/RigidFwdAsync.h
--- a/Engine/Plugins/Experimental/ChaosRigidPhysicsAsync/Source/ChaosRigidPhysicsAsync/Public/ChaosRigidPhysicsAsync/RigidFwdAsync.h
+++ b/Engine/Plugins/Experimental/ChaosRigidPhysicsAsync/Source/ChaosRigidPhysicsAsync/Public/ChaosRigidPhysicsAsync/RigidFwdAsync.h
@@ -5,8 +5,12 @@
 
 #if UE_RIGIDPHYSICS_API_ENABLED
 
-#include "Framework/Threading.h"
+namespace Chaos::Rigids::Async
+{
+	class FRigidSceneSettingsAsync;
+}
 
+// Internal fwd decls
 namespace Chaos::Rigids::Async
 {
 	class UE_INTERNAL FRigidBodyAsyncGT;
diff --git a/Engine/Plugins/Experimental/ChaosRigidPhysicsAsync/Source/ChaosRigidPhysicsAsync/Public/ChaosRigidPhysicsAsync/RigidSceneAsync.h b/Engine/Plugins/Experimental/ChaosRigidPhysicsAsync/Source/ChaosRigidPhysicsAsync/Public/ChaosRigidPhysicsAsync/RigidSceneAsync.h
--- a/Engine/Plugins/Experimental/ChaosRigidPhysicsAsync/Source/ChaosRigidPhysicsAsync/Public/ChaosRigidPhysicsAsync/RigidSceneAsync.h
+++ b/Engine/Plugins/Experimental/ChaosRigidPhysicsAsync/Source/ChaosRigidPhysicsAsync/Public/ChaosRigidPhysicsAsync/RigidSceneAsync.h
@@ -9,6 +9,7 @@
 #include "ChaosRigidPhysicsAsync/RigidBodyAsync.h"
 #include "ChaosRigidPhysicsAsync/RigidBodyContainerAsync.h"
 #include "ChaosRigidPhysicsAsync/RigidFactoryAsync.h"
+#include "ChaosRigidPhysicsAsync/RigidSceneSettingsAsync.h"
 #include "ChaosSolversModule.h"
 #include "Containers/Map.h"
 #include "Logging/StructuredLog.h"
@@ -33,6 +34,8 @@ namespace Chaos::Rigids::Async
 	class UE_INTERNAL FRigidSceneAsyncGT : public UE::Physics::FRigidScene, public TSharedFromThis<FRigidSceneAsyncGT, ESPMode::ThreadSafe>
 	{
 	public:
+		UE_RIGIDPHYSICS_RIGIDTYPED_DECL(CHAOSRIGIDPHYSICSASYNC_API, FRigidSceneAsyncGT);
+
 		UE_INTERNAL CHAOSRIGIDPHYSICSASYNC_API FRigidSceneAsyncGT(
 			const UE::Physics::FRigidDebugName& InName,
 			const FRigidSceneSettingsAsync& InSceneSettings);
@@ -103,6 +106,8 @@ namespace Chaos::Rigids::Async
 	class UE_INTERNAL FRigidSceneAsyncPT : public UE::Physics::IRigidScene
 	{
 	public:
+		UE_RIGIDPHYSICS_RIGIDTYPED_DECL(CHAOSRIGIDPHYSICSASYNC_API, FRigidSceneAsyncPT);
+
 		UE_INTERNAL CHAOSRIGIDPHYSICSASYNC_API FRigidSceneAsyncPT(const UE::Physics::FRigidDebugName& InName, const UE::Physics::FRigidBodyContainerHandle& InBodyPoolHandle);
 
 		//
diff --git a/Engine/Plugins/Experimental/ChaosRigidPhysicsAsync/Source/ChaosRigidPhysicsAsync/Public/ChaosRigidPhysicsAsync/RigidSceneSettingsAsync.h b/Engine/Plugins/Experimental/ChaosRigidPhysicsAsync/Source/ChaosRigidPhysicsAsync/Public/ChaosRigidPhysicsAsync/RigidSceneSettingsAsync.h
new file mode 100644
--- /dev/null
+++ b/Engine/Plugins/Experimental/ChaosRigidPhysicsAsync/Source/ChaosRigidPhysicsAsync/Public/ChaosRigidPhysicsAsync/RigidSceneSettingsAsync.h
@@ -0,0 +1,36 @@
+// Copyright Epic Games, Inc. All Rights Reserved.
+#pragma once
+
+#include "RigidPhysics/RigidFwd.h"
+
+#if UE_RIGIDPHYSICS_API_ENABLED
+
+#include "ChaosSolverConfiguration.h"
+#include "RigidPhysics/RigidSceneSettings.h"
+#include "PhysicsCoreTypes.h"
+
+namespace Chaos
+{
+	class FPBDRigidsSolver;
+	class FSingleParticlePhysicsProxy;
+}
+
+namespace Chaos::Rigids::Async
+{
+	class UE_INTERNAL FRigidSceneSettingsAsync : public UE::Physics::IRigidSceneSettings
+	{
+	public:
+		UE_RIGIDPHYSICS_RIGIDTYPED_DECL(CHAOSRIGIDPHYSICSASYNC_API, FRigidSceneSettingsAsync);
+
+		// The fixed timestep to use when async physics is enabled
+		float AsyncDt = -1.0f;
+
+		// The threading mode
+		EChaosThreadingMode ThreadingMode = EChaosThreadingMode::SingleThread;
+
+		// Solver settings
+		FChaosSolverConfiguration SolverConfig;
+	};
+}
+
+#endif // UE_RIGIDPHYSICS_API_ENABLED
\ No newline at end of file
diff --git a/Engine/Plugins/Experimental/ChaosRigidPhysicsAsync/Tests/Private/SceneLockBenchmarks.cpp b/Engine/Plugins/Experimental/ChaosRigidPhysicsAsync/Tests/Private/SceneLockBenchmarks.cpp
--- a/Engine/Plugins/Experimental/ChaosRigidPhysicsAsync/Tests/Private/SceneLockBenchmarks.cpp
+++ b/Engine/Plugins/Experimental/ChaosRigidPhysicsAsync/Tests/Private/SceneLockBenchmarks.cpp
@@ -3,6 +3,7 @@
 #include "Chaos/LowLevelTest/ChaosTestHarness.h"
 #include "Chaos/LowLevelTest/ChaosTestScene.h"
 #include "ChaosRigidPhysicsAsync/RigidFactoryAsync.h"
+#include "ChaosRigidPhysicsAsync/RigidSceneSettingsAsync.h"
 #include "ChaosRigidPhysicsAsync/RigidFwdAsync.h"
 #include "PBDRigidsSolver.h"
 #include "PhysicsInterfaceDeclaresCore.h"
@@ -27,7 +28,8 @@ namespace Chaos::LowLevelTest
 		FTransform Transform = FTransform::Identity;
 
 		FRigidFactoryAsync Factory;
-		FRigidSceneHandle SceneHandle = Factory.CreateScene(FRigidDebugName::Make(TEXT("Main")));
+		FRigidSceneSettingsAsync SceneSettings;
+		FRigidSceneHandle SceneHandle = Factory.CreateScene(FRigidDebugName::Make(TEXT("Main")), &SceneSettings);
 
 		// Create the bodies
 		TArray<FRigidBodyHandle> BodyHandles;
@@ -37,10 +39,10 @@ namespace Chaos::LowLevelTest
 			BodyHandles.Reserve(NumBodies);
 			for (int32 BodyIndex = 0; BodyIndex < NumBodies; ++BodyIndex)
 			{
-				TRigidBodyPtr<FRigidContextGameRW> Body = Context->CreateBody(FRigidDebugName::Make(TEXT("Body{0}"), BodyIndex), ERigidMovementType::Dynamic);
-				Body->Activate();
+				TRigidBodyPtr<FRigidContextGameRW> BodyPtr = Context->CreateBody(FRigidDebugName::Make(TEXT("Body{0}"), BodyIndex), ERigidMovementType::Dynamic);
+				BodyPtr->Activate();
 
-				BodyHandles.Add(Body);
+				BodyHandles.Add(BodyPtr);
 			}
 		}
 
@@ -129,7 +131,8 @@ namespace Chaos::LowLevelTest
 		FTransform Transform = FTransform::Identity;
 
 		FRigidFactoryAsync Factory;
-		FRigidSceneHandle SceneHandle = Factory.CreateScene(FRigidDebugName::Make(TEXT("Main")));
+		FRigidSceneSettingsAsync SceneSettings;
+		FRigidSceneHandle SceneHandle = Factory.CreateScene(FRigidDebugName::Make(TEXT("Main")), &SceneSettings);
 
 		// Create the bodies
 		TArray<FRigidBodyHandle> BodyHandles;
@@ -139,10 +142,10 @@ namespace Chaos::LowLevelTest
 			BodyHandles.Reserve(NumBodies);
 			for (int32 BodyIndex = 0; BodyIndex < NumBodies; ++BodyIndex)
 			{
-				TRigidBodyPtr<FRigidContextGameRW> Body = Context->CreateBody(FRigidDebugName::Make(TEXT("Body{0}"), BodyIndex), ERigidMovementType::Dynamic);
-				Body->Activate();
+				TRigidBodyPtr<FRigidContextGameRW> BodyPtr = Context->CreateBody(FRigidDebugName::Make(TEXT("Body{0}"), BodyIndex), ERigidMovementType::Dynamic);
+				BodyPtr->Activate();
 
-				BodyHandles.Add(Body);
+				BodyHandles.Add(BodyPtr);
 			}
 		}
 
diff --git a/Engine/Plugins/Experimental/ChaosRigidPhysicsAsync/Tests/Private/SceneTests.cpp b/Engine/Plugins/Experimental/ChaosRigidPhysicsAsync/Tests/Private/SceneTests.cpp
--- a/Engine/Plugins/Experimental/ChaosRigidPhysicsAsync/Tests/Private/SceneTests.cpp
+++ b/Engine/Plugins/Experimental/ChaosRigidPhysicsAsync/Tests/Private/SceneTests.cpp
@@ -6,6 +6,7 @@
 #include "ChaosRigidPhysicsAsyncTest.h"
 #include "ChaosRigidPhysicsAsync/RigidFactoryAsync.h"
 #include "ChaosRigidPhysicsAsync/RigidBodyAsync.h"
+#include "ChaosRigidPhysicsAsync/RigidSceneSettingsAsync.h"
 #include "Framework/Threading.h"
 #include "HAL/Thread.h"
 #include "RigidPhysics/RigidBody.h"
@@ -24,7 +25,8 @@ namespace Chaos::LowLevelTest
 	TEST_CASE("SceneLockTest", "[Chaos][API][Scene][unit]")
 	{
 		FRigidFactoryAsync SceneFactory;
-		FRigidSceneHandle SceneHandle = SceneFactory.CreateScene(FRigidDebugName::Make(TEXT("Main")));
+		FRigidSceneSettingsAsync SceneSettings;
+		FRigidSceneHandle SceneHandle = SceneFactory.CreateScene(FRigidDebugName::Make(TEXT("Main")), &SceneSettings);
 
 		// Create a thread that will lock the physics scene.
 		// We initially lock ThreadRunningLock on the main thread to prevent the thread from executing.
@@ -102,7 +104,8 @@ namespace Chaos::LowLevelTest
 		float Dt = 1.0f / 60.0f;
 
 		FRigidFactoryAsync SceneFactory;
-		FRigidSceneHandle SceneHandle = SceneFactory.CreateScene(FRigidDebugName::Make(TEXT("Main")));
+		FRigidSceneSettingsAsync SceneSettings;
+		FRigidSceneHandle SceneHandle = SceneFactory.CreateScene(FRigidDebugName::Make(TEXT("Main")), &SceneSettings);
 		REQUIRE(SceneHandle.LockRO());
 		if (FRigidContextGameRW Context = SceneHandle.LockRW())
 		{
@@ -210,6 +213,23 @@ namespace Chaos::LowLevelTest
 			}
 		}
 
+		SECTION("DestroyBodyInvalidatesPointer")
+		{
+			if (FRigidContextGameRW Context = SceneHandle.LockRW())
+			{
+				// Create a body
+				TRigidBodyPtr<FRigidContextGameRW> BodyPtr = Context->CreateBody(FRigidDebugName::Make(TEXT("Body1")), ERigidMovementType::Dynamic);
+				REQUIRE(BodyPtr.IsValid());
+
+				// The pointer should be invalidated by Destroy
+				Context->DestroyBody(MoveTemp(BodyPtr));
+				CHECK(!BodyPtr.IsValid());
+
+				// We gracefull handle attempt to destroy again (should not crash or assert)
+				Context->DestroyBody(MoveTemp(BodyPtr));
+			}
+		}
+
 		SceneFactory.DestroyScene(SceneHandle);
 	} // TEST_CASE
 
@@ -255,7 +275,8 @@ namespace Chaos::LowLevelTest
 			};
 
 		FRigidFactoryAsync SceneFactory;
-		FRigidSceneHandle SceneHandle = SceneFactory.CreateScene(FRigidDebugName::Make(TEXT("Main")));
+		FRigidSceneSettingsAsync SceneSettings;
+		FRigidSceneHandle SceneHandle = SceneFactory.CreateScene(FRigidDebugName::Make(TEXT("Main")), &SceneSettings);
 
 		if (FRigidContextGameRW Context = SceneHandle.LockRW())
 		{
@@ -324,7 +345,8 @@ namespace Chaos::LowLevelTest
 
 		// Create an async physics scene
 		FRigidFactoryAsync SceneFactory;
-		FRigidSceneHandle SceneHandle = SceneFactory.CreateScene(FRigidDebugName::Make(TEXT("Main")));
+		FRigidSceneSettingsAsync SceneSettings;
+		FRigidSceneHandle SceneHandle = SceneFactory.CreateScene(FRigidDebugName::Make(TEXT("Main")), &SceneSettings);
 
 		if (FRigidContextGameRW Context = SceneHandle.LockRW())
 		{
@@ -355,6 +377,24 @@ namespace Chaos::LowLevelTest
 		SceneFactory.DestroyScene(SceneHandle);
 	}
 
+
+	// Fake settings class to test scene creation failure
+	class FTestSceneSettings : public IRigidSceneSettings
+	{
+	public:
+		UE_RIGIDPHYSICS_RIGIDTYPED_DECL(CHAOSRIGIDPHYSICSASYNC_API, FTestSceneSettings);
+	};
+	UE_RIGIDPHYSICS_RIGIDTYPED_IMPL(FTestSceneSettings, IRigidSceneSettings);
+
+	// Attempt to create an async scene with settings for a different scene type. Should fail gracefully.
+	TEST_CASE("WrongSceneSettings", "[Chaos][API][Scene][unit]")
+	{
+		FRigidFactoryAsync SceneFactory;
+		FTestSceneSettings SceneSettings;
+		FRigidSceneHandle SceneHandle = SceneFactory.CreateScene(FRigidDebugName::Make(TEXT("Main")), &SceneSettings);
+		CHECK_FALSE(SceneHandle.LockRO());
+	}
+
 } // Chaos::LowLevelTest
 
-#endif
\ No newline at end of file
+#endif
diff --git a/Engine/Source/Runtime/Experimental/RigidPhysics/Private/RigidPhysics/RigidScene.cpp b/Engine/Source/Runtime/Experimental/RigidPhysics/Private/RigidPhysics/RigidScene.cpp
--- a/Engine/Source/Runtime/Experimental/RigidPhysics/Private/RigidPhysics/RigidScene.cpp
+++ b/Engine/Source/Runtime/Experimental/RigidPhysics/Private/RigidPhysics/RigidScene.cpp
@@ -4,12 +4,14 @@
 #include "RigidPhysics/RigidBodyContainer.h"
 #include "RigidPhysics/RigidContext.h"
 #include "RigidPhysics/RigidLog.h"
+#include "RigidPhysics/RigidSceneSettings.h"
 
 #if UE_RIGIDPHYSICS_API_ENABLED
 
 namespace UE::Physics
 {
 	UE_RIGIDPHYSICS_RIGIDTYPED_IMPL(IRigidScene, IRigidLockable);
+	UE_RIGIDPHYSICS_RIGIDTYPED_IMPL(IRigidSceneSettings, IRigidTyped);
 
 	FRigidSceneHandle::FRigidSceneHandle()
 		: SceneRaw(nullptr)
diff --git a/Engine/Source/Runtime/Experimental/RigidPhysics/Public/RigidPhysics/RigidFactory.h b/Engine/Source/Runtime/Experimental/RigidPhysics/Public/RigidPhysics/RigidFactory.h
new file mode 100644
--- /dev/null
+++ b/Engine/Source/Runtime/Experimental/RigidPhysics/Public/RigidPhysics/RigidFactory.h
@@ -0,0 +1,29 @@
+// Copyright Epic Games, Inc. All Rights Reserved.
+#pragma once
+
+#include "RigidPhysics/RigidFwd.h"
+
+#if UE_RIGIDPHYSICS_API_ENABLED
+
+#include "RigidPhysics/RigidTyped.h"
+
+namespace UE::Physics
+{
+	class IRigidFactory
+	{
+	public:
+		IRigidFactory()
+		{
+		}
+
+		virtual ~IRigidFactory()
+		{
+		}
+
+		virtual const FRigidTypeId& GetSceneTypeId() const = 0;
+		virtual FRigidSceneHandle CreateScene(const FRigidDebugName& InName, const IRigidSceneSettings* InSettings) = 0;
+		virtual void DestroyScene(const FRigidSceneHandle& InSceneHandle) = 0;
+	};
+}
+
+#endif
\ No newline at end of file
diff --git a/Engine/Source/Runtime/Experimental/RigidPhysics/Public/RigidPhysics/RigidFwd.h b/Engine/Source/Runtime/Experimental/RigidPhysics/Public/RigidPhysics/RigidFwd.h
--- a/Engine/Source/Runtime/Experimental/RigidPhysics/Public/RigidPhysics/RigidFwd.h
+++ b/Engine/Source/Runtime/Experimental/RigidPhysics/Public/RigidPhysics/RigidFwd.h
@@ -13,6 +13,7 @@
 #if UE_RIGIDPHYSICS_API_ENABLED
 
 #include "Chaos/ChaosDebugName.h"
+#include "Templates/RefCounting.h"
 #include "Templates/SharedPointer.h"
 
 #if DO_CHECK
@@ -45,8 +46,8 @@
 
 // Math/Bounds.h
 template<typename T> struct TBounds;
-using FBounds3f = TBounds< float >;
-using FBounds3d = TBounds< double >;
+using FBounds3f = TBounds<float>;
+using FBounds3d = TBounds<double>;
 
 // User-facing classes
 namespace UE::Physics
@@ -65,6 +66,7 @@ namespace UE::Physics
 	class IRigidScene;
 	class IRigidSceneFactory;
 	class IRigidSceneModifier;
+	class IRigidSceneSettings;
 	class IRigidTyped;
 
 	template<typename ContextType> class TRigidBodyContainerPtrImpl;
@@ -76,10 +78,7 @@ namespace UE::Physics
 	template<typename ContextType> using TRigidBodyPtr = TRigidObjectPtr<TRigidBodyPtrImpl<ContextType>>;
 	template<typename ContextType> using TRigidScenePtr = TRigidObjectPtr<TRigidScenePtrImpl<ContextType>>;
 
-	using FRigidSceneSharedPtr = TSharedPtr<IRigidScene, ESPMode::ThreadSafe>;
 	using FRigidSceneWeakPtr = TWeakPtr<IRigidScene, ESPMode::ThreadSafe>;
-	using FRigidBodyContainerSharedPtr = TSharedPtr<IRigidBodyContainer, ESPMode::ThreadSafe>;
-	using FRigidBodyContainerWeakPtr = TWeakPtr<IRigidBodyContainer, ESPMode::ThreadSafe>;
 
 	using FRigidDebugName = Chaos::FSharedDebugName;
 }
diff --git a/Engine/Source/Runtime/Experimental/RigidPhysics/Public/RigidPhysics/RigidSceneSettings.h b/Engine/Source/Runtime/Experimental/RigidPhysics/Public/RigidPhysics/RigidSceneSettings.h
new file mode 100644
--- /dev/null
+++ b/Engine/Source/Runtime/Experimental/RigidPhysics/Public/RigidPhysics/RigidSceneSettings.h
@@ -0,0 +1,22 @@
+// Copyright Epic Games, Inc. All Rights Reserved.
+#pragma once
+
+#include "RigidPhysics/RigidFwd.h"
+
+#if UE_RIGIDPHYSICS_API_ENABLED
+
+#include "RigidPhysics/RigidTyped.h"
+
+namespace UE::Physics
+{
+	class IRigidSceneSettings : public IRigidTyped
+	{
+	public:
+		UE_RIGIDPHYSICS_RIGIDTYPED_DECL(RIGIDPHYSICS_API, IRigidSceneSettings);
+
+		IRigidSceneSettings() {}
+		virtual ~IRigidSceneSettings() {}
+	};
+}
+
+#endif
\ No newline at end of file
