Chaos API
- Collison Shapes WIP

#rnx
#Synced-CL 48825514

diff --git a/Engine/Plugins/Experimental/ChaosRigidPhysicsAsync/Source/ChaosRigidPhysicsAsync/Private/ChaosRigidPhysicsAsync/RigidBodyAsync.cpp b/Engine/Plugins/Experimental/ChaosRigidPhysicsAsync/Source/ChaosRigidPhysicsAsync/Private/ChaosRigidPhysicsAsync/RigidBodyAsync.cpp
--- a/Engine/Plugins/Experimental/ChaosRigidPhysicsAsync/Source/ChaosRigidPhysicsAsync/Private/ChaosRigidPhysicsAsync/RigidBodyAsync.cpp
+++ b/Engine/Plugins/Experimental/ChaosRigidPhysicsAsync/Source/ChaosRigidPhysicsAsync/Private/ChaosRigidPhysicsAsync/RigidBodyAsync.cpp
@@ -3,7 +3,9 @@
 
 #include "ChaosRigidPhysicsAsync/RigidBodyContainerAsync.h"
 #include "ChaosRigidPhysicsAsync/RigidSceneAsync.h"
+#include "Chaos/ShapeInstance.h"
 #include "RigidPhysics/RigidBody.h"
+#include "RigidPhysics/RigidCollider.h"
 
 #if UE_RIGIDPHYSICS_API_ENABLED
 
@@ -11,6 +13,23 @@ namespace Chaos::Rigids::Async
 {
 	UE_RIGIDPHYSICS_RIGIDTYPED_IMPL(FRigidBodyAsyncGT, UE::Physics::IRigidBody);
 
+	double SafeInvert(double InValue, double InDefaultInverse = 0.0)
+	{
+		if (InValue != 0.0)
+		{
+			return 1.0 / InValue;
+		}
+		return InDefaultInverse;
+	}
+
+	FVector3d SafeInvert(const FVector3d& InValue, double InDefaultInverse = 0.0)
+	{
+		return FVector3d(
+			SafeInvert(InValue.X, InDefaultInverse),
+			SafeInvert(InValue.Y, InDefaultInverse),
+			SafeInvert(InValue.Z, InDefaultInverse));
+	}
+
 	FRigidBodyAsyncGT::FRigidBodyAsyncGT(FRigidBodyPoolAsyncGT* InContainer, const UE::Physics::FRigidDebugName& InName, UE::Physics::ERigidMovementType InMovementType)
 		: BodyId()
 		, BodyContainer(InContainer)
@@ -40,13 +59,6 @@ namespace Chaos::Rigids::Async
 #if CHAOS_DEBUG_NAME
 			Particle->SetDebugName(InName.SharedValue());
 #endif
-
-			// Create a default shape for testing
-			// TODO_CHAOS_API: Pass body settings into factory functions
-			FImplicitObjectPtr Geom = FImplicitObjectPtr(new FImplicitBox3(FVec3(-50), FVec3(50)));
-			FImplicitObjectPtr Union = new FImplicitObjectUnion(TArray<FImplicitObjectPtr>({ Geom }));
-			Particle->SetGeometry(Union);
-
 			Proxy = FSingleParticlePhysicsProxy::Create(MoveTemp(Particle));
 		}
 	}
@@ -160,6 +172,38 @@ namespace Chaos::Rigids::Async
 		return (GetProxyAPI().ObjectState() == EObjectStateType::Sleeping);
 	}
 
+	int32 FRigidBodyAsyncGT::GetNumShapes() const
+	{
+		return GetProxyAPI().ShapesArray().Num();
+	}
+
+	UE::Physics::FConstRigidShapeInstancePtr FRigidBodyAsyncGT::GetShape(int32 InShapeIndex) const
+	{
+		using namespace UE::Physics;
+
+		if (GetProxyAPI().ShapesArray().IsValidIndex(InShapeIndex))
+		{
+			return FConstRigidShapeInstancePtr(GetProxyAPI().ShapesArray()[InShapeIndex].Get());
+		}
+		return FConstRigidShapeInstancePtr();
+	}
+
+	void FRigidBodyAsyncGT::AddShape(const UE::Physics::FConstRigidShapeInstancePtr& InShape)
+	{
+		using namespace Chaos;
+
+		TArray<FImplicitObjectPtr> ImplicitsObjects;
+		ImplicitsObjects = { InShape->GetPerShapeData()->GetGeometry() };
+		GetProxyAPI().MergeGeometry(MoveTemp(ImplicitsObjects));
+
+		// TODO_CHAOSAPI: We are wrapping the collider's PerShapeData in a UniquePtr here
+		// which is leaving a raw pointer to the same object behind in InShape.
+		// This is temporary until the API is used everywhere and we can refactor.
+		FShapesArray Shapes;
+		Shapes.Add(TUniquePtr<FPerShapeData>(InShape->GetPerShapeData()));
+		GetProxyAPI().MergeShapesArray(MoveTemp(Shapes));
+	}
+
 	void FRigidBodyAsyncGT::InitTransform(const FTransform& InTransform)
 	{
 		// TODO_CHAOSAPI: This is not supported by the Proxy on the GT. 
@@ -174,13 +218,13 @@ namespace Chaos::Rigids::Async
 		FRigidBodyHandle_External& ProxyAPI = GetProxyAPI();
 		if (!IsKinematic() && !IsSleeping() && FVec3::IsNearlyEqual(InTransform.GetLocation(), ProxyAPI.X(), SMALL_NUMBER) && FRotation3::IsNearlyEqual(InTransform.GetRotation(), ProxyAPI.R(), SMALL_NUMBER))
 		{
-			// if simulating, don't update X/R if they haven't changed. this allows scale to be set on simulating body without overriding async position/rotation.
+			// Don't update X/R if they haven't changed. This allows scale to be set on simulating body without overriding async position/rotation.
 			return;
 		}
 
 		if (IsKinematic())
 		{
-			// NOTE: SetGlobalPose is a teleport for kinematics. Use SetKinematicTarget_AssumesLocked
+			// NOTE: UpdateTransform is a teleport for kinematics. Use SetKinematicTarget
 			// if the kinematic should calculate its velocity from the transform delta.
 			ProxyAPI.SetKinematicTarget(InTransform);
 			ProxyAPI.SetV(FVector::Zero());
@@ -190,6 +234,12 @@ namespace Chaos::Rigids::Async
 		ProxyAPI.SetX(InTransform.GetTranslation());
 		ProxyAPI.SetR(InTransform.GetRotation());
 		ProxyAPI.UpdateShapeBounds();
+
+		// TODO_CHAOSAPI: when we have the Spatial Acceleration working
+		//if (FRigidSceneAsyncGT* Scene = GetScene())
+		//{
+		//	Scene->OnRigidBodyTransformUpdated(this);
+		//}
 	}
 
 	void FRigidBodyAsyncGT::SetKinematicTarget(const FTransform& InTransform)
@@ -230,6 +280,18 @@ namespace Chaos::Rigids::Async
 	{
 		UE_RIGIDPHYSICS_CHECK(IsValid());
 		GetProxyAPI().SetM(InMass);
+		GetProxyAPI().SetInvM(SafeInvert(InMass));
+	}
+
+	FVector3d FRigidBodyAsyncGT::GetInertia() const
+	{
+		return FVector3d(GetProxyAPI().I());
+	}
+
+	void FRigidBodyAsyncGT::SetInertia(const FVector3d& InInertia)
+	{
+		GetProxyAPI().SetI(FVec3f(InInertia));
+		GetProxyAPI().SetInvI(FVec3f(SafeInvert(InInertia)));
 	}
 
 	//void SetCCDEnabled(bool bInEnabled) {  }
@@ -301,6 +363,28 @@ namespace Chaos::Rigids::Async
 		return (GetProxyAPI()->ObjectState() == EObjectStateType::Sleeping);
 	}
 
+	int32 FRigidBodyAsyncPT::GetNumShapes() const
+	{
+		return GetProxyAPI()->ShapesArray().Num();
+	}
+
+	UE::Physics::FConstRigidShapeInstancePtr FRigidBodyAsyncPT::GetShape(int32 InShapeIndex) const
+	{
+		using namespace UE::Physics;
+
+		if (GetProxyAPI()->ShapesArray().IsValidIndex(InShapeIndex))
+		{
+			return FConstRigidShapeInstancePtr(GetProxyAPI()->ShapesArray()[InShapeIndex].Get());
+		}
+
+		return FConstRigidShapeInstancePtr();
+	}
+
+	void FRigidBodyAsyncPT::AddShape(const UE::Physics::FConstRigidShapeInstancePtr& InCollider)
+	{
+		UE_RIGIDPHYSICS_CHECK(false);
+	}
+
 	void FRigidBodyAsyncPT::InitTransform(const FTransform& InTransform)
 	{
 		SetTransform(InTransform,  /*bIsTeleport*/true);
diff --git a/Engine/Plugins/Experimental/ChaosRigidPhysicsAsync/Source/ChaosRigidPhysicsAsync/Public/ChaosRigidPhysicsAsync/RigidBodyAsync.h b/Engine/Plugins/Experimental/ChaosRigidPhysicsAsync/Source/ChaosRigidPhysicsAsync/Public/ChaosRigidPhysicsAsync/RigidBodyAsync.h
--- a/Engine/Plugins/Experimental/ChaosRigidPhysicsAsync/Source/ChaosRigidPhysicsAsync/Public/ChaosRigidPhysicsAsync/RigidBodyAsync.h
+++ b/Engine/Plugins/Experimental/ChaosRigidPhysicsAsync/Source/ChaosRigidPhysicsAsync/Public/ChaosRigidPhysicsAsync/RigidBodyAsync.h
@@ -49,6 +49,9 @@ namespace Chaos::Rigids::Async
 		UE_INTERNAL virtual bool IsKinematic() const override final;
 		UE_INTERNAL virtual bool IsDynamic() const override final;
 		UE_INTERNAL virtual bool IsSleeping() const override final;
+		UE_INTERNAL virtual int32 GetNumShapes() const override final;
+		UE_INTERNAL virtual UE::Physics::FConstRigidShapeInstancePtr GetShape(int32 InShapeIndex) const override final;
+		UE_INTERNAL virtual void AddShape(const UE::Physics::FConstRigidShapeInstancePtr& InCollider) override final;
 		UE_INTERNAL virtual void InitTransform(const FTransform& InTransform) override final;
 		UE_INTERNAL virtual void UpdateTransform(const FTransform& InTransform) override final;
 		UE_INTERNAL virtual void SetKinematicTarget(const FTransform& InTransform) override final;
@@ -56,6 +59,8 @@ namespace Chaos::Rigids::Async
 		UE_INTERNAL virtual FBounds3d GetBounds() const override final;
 		UE_INTERNAL virtual double GetMass() const override final;
 		UE_INTERNAL virtual void SetMass(double InMass) override final;
+		UE_INTERNAL virtual FVector3d GetInertia() const override final;
+		UE_INTERNAL virtual void SetInertia(const FVector3d& InInertia) override final;
 		//UE_INTERNAL virtual void SetCCDEnabled(bool bInEnabled) override final {  }
 		//UE_INTERNAL virtual void SetMACDEnabled(bool bInEnabled) override final {  }
 		//UE_INTERNAL virtual void SetSmoothEdgeCollisionsEnabled(bool bInEnabled) override final {  }
@@ -120,6 +125,9 @@ namespace Chaos::Rigids::Async
 		UE_INTERNAL virtual bool IsKinematic() const override final;
 		UE_INTERNAL virtual bool IsDynamic() const override final;
 		UE_INTERNAL virtual bool IsSleeping() const override final;
+		UE_INTERNAL virtual int32 GetNumShapes() const override final;
+		UE_INTERNAL virtual UE::Physics::FConstRigidShapeInstancePtr GetShape(int32 InShapeIndex) const override final;
+		UE_INTERNAL virtual void AddShape(const UE::Physics::FConstRigidShapeInstancePtr& InCollider) override final;
 		UE_INTERNAL virtual void InitTransform(const FTransform& InTransform) override final;
 		UE_INTERNAL virtual void UpdateTransform(const FTransform& InTransform) override final;
 		UE_INTERNAL virtual void SetKinematicTarget(const FTransform& InTransform) override final;
diff --git a/Engine/Plugins/Experimental/ChaosRigidPhysicsAsync/Tests/Private/BodyTests.cpp b/Engine/Plugins/Experimental/ChaosRigidPhysicsAsync/Tests/Private/BodyTests.cpp
new file mode 100644
--- /dev/null
+++ b/Engine/Plugins/Experimental/ChaosRigidPhysicsAsync/Tests/Private/BodyTests.cpp
@@ -0,0 +1,262 @@
+// Copyright Epic Games, Inc. All Rights Reserved.
+
+#include "Chaos/LowLevelTest/ChaosTestHarness.h"
+#include "Chaos/LowLevelTest/ChaosTestScene.h"
+#include "ChaosRigidPhysicsAsyncTest.h"
+#include "ChaosRigidPhysicsAsync/RigidFactoryAsync.h"
+#include "ChaosRigidPhysicsAsync/RigidBodyAsync.h"
+#include "ChaosRigidPhysicsAsync/RigidSceneAsync.h"
+#include "ChaosRigidPhysicsAsync/RigidSceneSettingsAsync.h"
+#include "RigidPhysics/RigidBody.h"
+#include "RigidPhysics/RigidCollider.h"
+#include "RigidPhysics/RigidContext.h"
+#include "RigidPhysics/RigidModifier.h"
+#include "RigidPhysics/RigidScene.h"
+
+#if UE_RIGIDPHYSICS_API_ENABLED
+
+namespace Chaos::LowLevelTest
+{
+	using namespace Chaos;
+	using namespace Chaos::Rigids::Async;
+	using namespace UE::Physics;
+
+	// Create bodies with various setups and validate them
+	TEST_CASE("BodySettingsTest", "[Chaos][API][Body][unit]")
+	{
+		FRigidFactoryAsync Factory;
+		FRigidSceneSettingsAsync SceneSettings;
+		FRigidSceneHandle SceneHandle = Factory.CreateScene(FRigidDebugName::Make(TEXT("Main")), &SceneSettings);
+
+		FRigidBodyHandle Body1Handle;
+		FRigidBodyHandle Body2Handle;
+		if (FRigidContextGameRW Context = SceneHandle.LockRW())
+		{
+			TRigidBodyPtr<FRigidContextGameRW> Body1Ptr = Context->CreateBody(FRigidDebugName("Body1"), ERigidMovementType::Static);
+			REQUIRE(Body1Ptr);
+			Body1Ptr->AddShape(MakeBoxShape(FVector3f(1000, 1000, 100)));
+			Body1Ptr->InitTransform(FTransform(FVector(0, 0, -50)));
+			Body1Ptr->Activate();
+
+			TRigidBodyPtr<FRigidContextGameRW> Body2Ptr = Context->CreateBody(FRigidDebugName("Body2"), ERigidMovementType::Dynamic);
+			REQUIRE(Body2Ptr);
+			Body2Ptr->AddShape(MakeBoxShape(FVector3f(100, 100, 100)));
+			Body2Ptr->SetIsSleeping(false);
+			Body2Ptr->SetMass(1000);
+			Body2Ptr->SetInertia(FVector3d(MakeSolidBoxInertia(1000, FVector3f(100, 100, 100))));
+			Body2Ptr->InitTransform(FTransform(FVector(0, 0, 300)));
+			Body2Ptr->Activate();
+
+			// Check the bodies are where we think they are and the bounds are up to date
+			FBounds3d Body1Bounds = Body1Ptr->GetBounds();
+			CHECK(FVector3d(Body1Bounds.Min) == FVector3d(-500, -500, -100));
+			CHECK(FVector3d(Body1Bounds.Max) == FVector3d(500, 500, 0));
+
+			FBounds3d Body2Bounds = Body2Ptr->GetBounds();
+			CHECK(FVector3d(Body2Bounds.Min) == FVector3d(-50, -50, 250));
+			CHECK(FVector3d(Body2Bounds.Max) == FVector3d(50, 50, 350));
+
+			Body1Handle = Body1Ptr;
+			Body2Handle = Body2Ptr;
+		}
+
+		if (FRigidContextGameRW Context = SceneHandle.LockRW())
+		{
+			TRigidBodyPtr<FRigidContextGameRW> Body2Ptr = Body2Handle.Pin(Context);
+			REQUIRE(Body2Ptr);
+
+			for (int32 TickIndex = 0; TickIndex < 100; ++TickIndex)
+			{
+				Context->Tick(0.01);
+			}
+
+			CHECK_THAT(Body2Ptr->GetTransform().GetLocation(), Catch::ApproxEq(FVec3(0.0, 0.0, 50.0)));
+		}
+
+		if (FRigidContextGameRW Context = SceneHandle.LockRW())
+		{
+			Context->DestroyBody(Body1Handle.Pin(Context));
+			Context->DestroyBody(Body2Handle.Pin(Context));
+		}
+
+		Factory.DestroyScene(SceneHandle);
+	}
+}
+
+#endif
+
+
+#if 0
+
+template<typename InAllocatorType>
+void TInitBodiesHelperBase<InAllocatorType>::CreateActor_AssumesLocked(FBodyInstance* Instance, const FTransform& Transform) const
+{
+	SCOPE_CYCLE_COUNTER(STAT_CreatePhysicsActor);
+	checkSlow(!FPhysicsInterface::IsValid(Instance->GetPhysicsActor()));
+	const ECollisionEnabled::Type CollisionType = Instance->GetCollisionEnabled();
+
+	FActorCreationParams ActorParams;
+	ActorParams.InitialTM = Transform;
+	ActorParams.bSimulatePhysics = Instance->ShouldInstanceSimulatingPhysics();
+	ActorParams.bStartAwake = Instance->bStartAwake;
+#if USE_BODYINSTANCE_DEBUG_NAMES
+	ActorParams.DebugName = Instance->CharDebugName.IsValid() ? Instance->CharDebugName->GetData() : nullptr;
+#endif
+	ActorParams.bEnableGravity = Instance->bEnableGravity;
+	ActorParams.bUpdateKinematicFromSimulation = Instance->bUpdateKinematicFromSimulation;
+	ActorParams.bQueryOnly = CollisionType == ECollisionEnabled::QueryOnly;
+	ActorParams.Scene = PhysScene;
+	ActorParams.bStatic = IsStatic();
+
+	FPhysicsActorHandle ActorHandle = nullptr;
+	FPhysicsInterface::CreateActor(ActorParams, ActorHandle, PrimitiveComp);
+	Instance->SetPhysicsActor(ActorHandle);
+
+	if (!IsStatic())
+	{
+		FPhysicsInterface::SetCcdEnabled_AssumesLocked(ActorHandle, Instance->bUseCCD);
+		FPhysicsInterface::SetMACDEnabled_AssumesLocked(ActorHandle, Instance->IsUsingMACD());
+		FPhysicsInterface::SetPartialIslandSleepAllowed_AssumesLocked(ActorHandle, Instance->IsPartialIslandSleepAllowed());
+		FPhysicsInterface::SetIsKinematic_AssumesLocked(ActorHandle, !Instance->ShouldInstanceSimulatingPhysics());
+
+		FPhysicsInterface::SetMaxLinearVelocity_AssumesLocked(ActorHandle, TNumericLimits<float>::Max());
+		FPhysicsInterface::SetSmoothEdgeCollisionsEnabled_AssumesLocked(ActorHandle, Instance->bSmoothEdgeCollisions);
+		FPhysicsInterface::SetInertiaConditioningEnabled_AssumesLocked(ActorHandle, Instance->IsInertiaConditioningEnabled());
+		FPhysicsInterface::SetGyroscopicTorqueEnabled_AssumesLocked(ActorHandle, Instance->bGyroscopicTorqueEnabled);
+		FPhysicsInterface::SetPositionSolverIterationCount_AssumesLocked(ActorHandle, Instance->GetPositionSolverIterationCount());
+		FPhysicsInterface::SetVelocitySolverIterationCount_AssumesLocked(ActorHandle, Instance->GetVelocitySolverIterationCount());
+		FPhysicsInterface::SetProjectionSolverIterationCount_AssumesLocked(ActorHandle, Instance->GetProjectionSolverIterationCount());
+		FPhysicsInterface::SetGravityGroupIndex_AssumesLocked(ActorHandle, Instance->GravityGroupIndex);
+
+		// Set sleep event notification
+		FPhysicsInterface::SetSendsSleepNotifies_AssumesLocked(ActorHandle, Instance->bGenerateWakeEvents);
+	}
+}
+
+template<typename InAllocatorType>
+bool TInitBodiesHelperBase<InAllocatorType>::CreateShapes_AssumesLocked(FBodyInstance* Instance) const
+{
+	SCOPE_CYCLE_COUNTER(STAT_CreatePhysicsShapes);
+	UPhysicalMaterial* SimplePhysMat = Instance->GetSimplePhysicalMaterial();
+	TArray<UPhysicalMaterial*> ComplexPhysMats;
+	TArray<FPhysicalMaterialMaskParams> ComplexPhysMatMasks;
+
+	ComplexPhysMats = Instance->GetComplexPhysicalMaterials(ComplexPhysMatMasks);
+
+	FBodyCollisionData BodyCollisionData;
+	Instance->BuildBodyFilterData(BodyCollisionData.CollisionFilterData);
+	FBodyInstance::BuildBodyCollisionFlags(BodyCollisionData.CollisionFlags, Instance->GetCollisionEnabled(), BodySetup->GetCollisionTraceFlag() == CTF_UseComplexAsSimple);
+
+	bool bInitFail = false;
+
+	// #PHYS2 Call interface AddGeometry
+	BodySetup->AddShapesToRigidActor_AssumesLocked(Instance, Instance->Scale3D, SimplePhysMat, ComplexPhysMats, ComplexPhysMatMasks, BodyCollisionData, FTransform::Identity);
+
+	FPhysicsInterface::SetIgnoreAnalyticCollisions_AssumesLocked(Instance->GetPhysicsActor(), CVarIgnoreAnalyticCollisionsOverride.GetValueOnAnyThread() ? true : Instance->bIgnoreAnalyticCollisions);
+
+	const int32 NumShapes = FPhysicsInterface::GetNumShapes(Instance->GetPhysicsActor());
+	bInitFail |= NumShapes == 0;
+
+	return bInitFail;
+}
+
+void UBodySetup::AddShapesToRigidActor_AssumesLocked(
+	FBodyInstance* OwningInstance,
+	FVector& Scale3D,
+	UPhysicalMaterial* SimpleMaterial,
+	TArray<UPhysicalMaterial*>& ComplexMaterials,
+	TArray<FPhysicalMaterialMaskParams>& ComplexMaterialMasks,
+	const FBodyCollisionData& BodyCollisionData,
+	const FTransform& RelativeTM,
+	TArray<FPhysicsShapeHandle>* NewShapes)
+{
+	SCOPE_CYCLE_COUNTER(STAT_AddShapesToActor);
+
+	check(OwningInstance);
+
+	// in editor, there are a lot of things relying on body setup to create physics meshes
+	CreatePhysicsMeshes();
+
+	// To AddGeometry in interface
+	// if almost zero, set min scale
+	// @todo fixme
+	if (Scale3D.IsNearlyZero())
+	{
+		// set min scale
+		Scale3D = FVector(0.1f);
+	}
+
+	FGeometryAddParams AddParams;
+	AddParams.bDoubleSided = bDoubleSidedGeometry;
+	AddParams.CollisionData = BodyCollisionData;
+	AddParams.CollisionTraceType = GetCollisionTraceFlag();
+	AddParams.Scale = Scale3D;
+	AddParams.SimpleMaterial = SimpleMaterial;
+	AddParams.ComplexMaterials = TArrayView<UPhysicalMaterial*>(ComplexMaterials);
+	AddParams.ComplexMaterialMasks = TArrayView<FPhysicalMaterialMaskParams>(ComplexMaterialMasks);
+	AddParams.LocalTransform = RelativeTM;
+	AddParams.WorldTransform = OwningInstance->GetUnrealWorldTransform();
+	AddParams.Geometry = &AggGeom;
+	AddParams.TriMeshGeometries = MakeArrayView(TriMeshGeometries);
+
+	{
+		SCOPE_CYCLE_COUNTER(STAT_AddGeomToSolver);
+		FPhysicsInterface::AddGeometry(OwningInstance->GetPhysicsActor(), AddParams, NewShapes);
+	}
+}
+
+void FPhysInterface_Chaos::AddGeometry(const FPhysicsActorHandle& InActor, const FGeometryAddParams& InParams, TArray<FPhysicsShapeHandle>* OutOptShapes)
+{
+	TRACE_CPUPROFILER_EVENT_SCOPE(FPhysInterface_Chaos::AddGeometry);
+	LLM_SCOPE(ELLMTag::ChaosGeometry);
+
+	// @todo(chaos): we should not be creating unique geometry per actor
+	// @todo(chaos): we are creating the Shapes array twice. Once here and again in SetGeometry or MergeGeometry. Fix this.
+	TArray<Chaos::FImplicitObjectPtr> Geoms;
+	Chaos::FShapesArray Shapes;
+	ChaosInterface::CreateGeometry(InParams, Geoms, Shapes);
+
+	if (InActor && Geoms.Num())
+	{
+		for (TUniquePtr<Chaos::FPerShapeData>& Shape : Shapes)
+		{
+			FPhysicsShapeHandle NewHandle(Shape.Get(), InActor);
+			if (OutOptShapes)
+			{
+				OutOptShapes->Add(NewHandle);
+			}
+
+			FBodyInstance::ApplyMaterialToShape_AssumesLocked(NewHandle, InParams.SimpleMaterial, InParams.ComplexMaterials, &InParams.ComplexMaterialMasks);
+
+			//TArrayView<UPhysicalMaterial*> SimpleView = MakeArrayView(&(const_cast<UPhysicalMaterial*>(InParams.SimpleMaterial)), 1);
+			//FPhysInterface_Chaos::SetMaterials(NewHandle, InParams.ComplexMaterials.Num() > 0 ? InParams.ComplexMaterials : SimpleView);
+		}
+
+		// NOTE: Both MergeGeometry and SetGeometry will extend the ShapesInstances array to contain enough elements for
+		// each geometry in the Union. However the shape data will not have been filled in, hence the call to MergeShapeInstance at the end.
+		// todo: we should not be creating unique geometry per actor
+		{
+			if (InActor->GetGameThreadAPI().GetGeometry())
+			{
+				// Geometry already exists - combine new geometry with the existing
+				// NOTE: We do not need to set the AllowBVH flag because it will be cloned (see below)
+				InActor->GetGameThreadAPI().MergeGeometry(MoveTemp(Geoms));
+			}
+			else
+			{
+				// We always have a union so we can support any future welding operations. (Non-trivial converting the SharedPtr to UniquePtr).
+				// NOTE: The root union always supports BVH (if there are enough shapes) and is the only Union in the hierarchy that is allowed 
+				// to do so, but we don't create it here because that makes welding even more expensive (bodies are welded one by one). 
+				// Search for SetAllowBVH to see where the BVH is enabled.
+				Chaos::FImplicitObjectPtr Union = MakeImplicitObjectPtr<Chaos::FImplicitObjectUnion>(MoveTemp(Geoms));
+				InActor->GetGameThreadAPI().SetGeometry(MoveTemp(Union));
+			}
+		}
+
+		// Update the newly added shapes with the collision filters, materials etc
+		// NOTE: MergeShapes overwrites the last N shapes (see comments above)
+		InActor->GetGameThreadAPI().MergeShapesArray(MoveTemp(Shapes));
+	}
+}
+
+#endif
diff --git a/Engine/Plugins/Experimental/ChaosRigidPhysicsAsync/Tests/Private/SceneLockBenchmarks.cpp b/Engine/Plugins/Experimental/ChaosRigidPhysicsAsync/Tests/Private/SceneLockBenchmarks.cpp
--- a/Engine/Plugins/Experimental/ChaosRigidPhysicsAsync/Tests/Private/SceneLockBenchmarks.cpp
+++ b/Engine/Plugins/Experimental/ChaosRigidPhysicsAsync/Tests/Private/SceneLockBenchmarks.cpp
@@ -9,6 +9,7 @@
 #include "PhysicsInterfaceDeclaresCore.h"
 #include "Physics/Experimental/PhysInterface_Chaos.h"
 #include "PhysicsProxy/SingleParticlePhysicsProxy.h"
+#include "RigidPhysics/RigidCollider.h"
 #include "RigidPhysics/RigidContext.h"
 
 #if UE_RIGIDPHYSICS_API_ENABLED
@@ -40,6 +41,7 @@ namespace Chaos::LowLevelTest
 			for (int32 BodyIndex = 0; BodyIndex < NumBodies; ++BodyIndex)
 			{
 				TRigidBodyPtr<FRigidContextGameRW> BodyPtr = Context->CreateBody(FRigidDebugName::Make(TEXT("Body{0}"), BodyIndex), ERigidMovementType::Dynamic);
+				BodyPtr->AddShape(MakeBoxShape(FVector3f(50, 50, 50)));
 				BodyPtr->Activate();
 
 				BodyHandles.Add(BodyPtr);
@@ -143,6 +145,7 @@ namespace Chaos::LowLevelTest
 			for (int32 BodyIndex = 0; BodyIndex < NumBodies; ++BodyIndex)
 			{
 				TRigidBodyPtr<FRigidContextGameRW> BodyPtr = Context->CreateBody(FRigidDebugName::Make(TEXT("Body{0}"), BodyIndex), ERigidMovementType::Dynamic);
+				BodyPtr->AddShape(MakeBoxShape(FVector3f(50, 50, 50)));
 				BodyPtr->Activate();
 
 				BodyHandles.Add(BodyPtr);
diff --git a/Engine/Plugins/Experimental/ChaosRigidPhysicsAsync/Tests/Private/SceneTests.cpp b/Engine/Plugins/Experimental/ChaosRigidPhysicsAsync/Tests/Private/SceneTests.cpp
--- a/Engine/Plugins/Experimental/ChaosRigidPhysicsAsync/Tests/Private/SceneTests.cpp
+++ b/Engine/Plugins/Experimental/ChaosRigidPhysicsAsync/Tests/Private/SceneTests.cpp
@@ -10,6 +10,7 @@
 #include "Framework/Threading.h"
 #include "HAL/Thread.h"
 #include "RigidPhysics/RigidBody.h"
+#include "RigidPhysics/RigidCollider.h"
 #include "RigidPhysics/RigidContext.h"
 #include "RigidPhysics/RigidModifier.h"
 #include "RigidPhysics/RigidScene.h"
@@ -139,6 +140,10 @@ namespace Chaos::LowLevelTest
 				CHECK(StaticBodyHandle.Pin(Context)->GetTypeId() == FRigidBodyAsyncGT::GetStaticTypeId());
 				CHECK(DynamicBodyHandle.Pin(Context)->GetTypeId() == FRigidBodyAsyncGT::GetStaticTypeId());
 
+				// Add Collision
+				StaticBodyHandle.Pin(Context)->AddShape(MakeBoxShape(FVector3f(100, 100, 100)));
+				DynamicBodyHandle.Pin(Context)->AddShape(MakeBoxShape(FVector3f(100, 100, 100)));
+
 				// Activate Bodies
 				StaticBodyHandle.Pin(Context)->Activate();
 				DynamicBodyHandle.Pin(Context)->Activate();
@@ -202,6 +207,9 @@ namespace Chaos::LowLevelTest
 				CHECK(BodyHandleDupe.Pin(Context));
 				CHECK(BodyHandleDupe.Pin(Context).Get() == BodyHandle.Pin(Context).Get());
 
+				// Add Collision
+				BodyHandle.Pin(Context)->AddShape(MakeBoxShape(FVector3f(100, 100, 100)));
+
 				Context->Tick(Dt);
 
 				// Destroy the body
@@ -353,10 +361,12 @@ namespace Chaos::LowLevelTest
 			Context->RegisterModifier(&SceneModifier);
 
 			TRigidBodyPtr<FRigidContextGameRW> BodyPtr1 = Context->CreateBody(FRigidDebugName::Make(TEXT("Body1")), ERigidMovementType::Dynamic);
+			BodyPtr1->AddShape(MakeBoxShape(FVector3f(100, 100, 100)));
 			BodyPtr1->Activate();
 			BodyHandle1 = BodyPtr1;
 
 			TRigidBodyPtr<FRigidContextGameRW> BodyPtr2 = Context->CreateBody(FRigidDebugName::Make(TEXT("Body2")), ERigidMovementType::Dynamic);
+			BodyPtr2->AddShape(MakeBoxShape(FVector3f(100, 100, 100)));
 			BodyPtr2->Activate();
 			BodyHandle2 = BodyPtr2;
 
@@ -395,6 +405,67 @@ namespace Chaos::LowLevelTest
 		CHECK_FALSE(SceneHandle.LockRO());
 	}
 
+
+	//
+	// NOTE: NOT A REAL TEST
+	// This code is just sample code and mainly used to cut and paste images into the architecture presentation!
+	//
+	TEST_CASE("SceneExampleCode", "![SampleCodeOnly]")
+	{
+		FRigidFactoryAsync SceneFactory;
+		FRigidSceneSettingsAsync SceneSettings;
+		FRigidSceneHandle SceneHandle = SceneFactory.CreateScene(FRigidDebugName::Make(TEXT("Main")), &SceneSettings);
+
+		FRigidBodyHandle BodyHandle;
+		if (FRigidContextGameRW Context = SceneHandle.LockRW())
+		{
+			if (TRigidBodyPtr<FRigidContextGameRW> BodyPtr = Context->CreateBody(FRigidDebugName("Body2"), ERigidMovementType::Dynamic))
+			{
+				BodyPtr->AddShape(MakeBoxShape(FVector3f(100, 100, 100)));
+				BodyPtr->SetIsSleeping(false);
+				BodyPtr->SetMass(1000);
+				BodyPtr->SetInertia(FVector3d(MakeSolidBoxInertia(1000, FVector3f(100, 100, 100))));
+				BodyPtr->InitTransform(FTransform(FVector(0, 0, 300)));
+				BodyPtr->Activate();
+
+				BodyHandle = BodyPtr;
+			}
+		}
+
+		if (FRigidContextGameRW Context = SceneHandle.LockRW())
+		{
+			if (TRigidBodyPtr<FRigidContextGameRW> BodyPtr = BodyHandle.Pin(Context))
+			{
+				BodyPtr->SetMass(1000);
+				BodyPtr->InitTransform(FTransform(FVector(0, 0, 300)));
+			}
+		}
+
+		//FPhysicsActorHandle ActorHandle;
+		//FPhysicsCommand::ExecuteWrite(ActorHandle,
+		//	[](const FPhysicsActorHandle& Actor)
+		//	{
+		//		FPhysicsInterface::SetGlobalPose_AssumesLocked(Actor, FTransform(FVector(0, 0, 300)));
+		//		FPhysicsInterface::SetMass_AssumesLocked(Actor, 1000);
+		//	});
+
+		const auto& PTCallback = [&BodyHandle](const FRigidContextSimRW& Context)
+			{
+				if (TRigidBodyPtr<FRigidContextSimRW> BodyPtr = BodyHandle.Pin(Context))
+				{
+					BodyPtr->SetMass(1000);
+					BodyPtr->InitTransform(FTransform(FVector(0, 0, 300)));
+				}
+			};
+
+		if (FRigidContextGameRW Context = SceneHandle.LockRW())
+		{
+			Context->DestroyBody(BodyHandle.Pin(Context));
+		}
+
+		SceneFactory.DestroyScene(SceneHandle);
+	}
+
 } // Chaos::LowLevelTest
 
 #endif
diff --git a/Engine/Source/Runtime/Experimental/RigidPhysics/Private/RigidPhysics/RigidCollider.cpp b/Engine/Source/Runtime/Experimental/RigidPhysics/Private/RigidPhysics/RigidCollider.cpp
new file mode 100644
--- /dev/null
+++ b/Engine/Source/Runtime/Experimental/RigidPhysics/Private/RigidPhysics/RigidCollider.cpp
@@ -0,0 +1,137 @@
+// Copyright Epic Games, Inc. All Rights Reserved.
+#include "RigidPhysics/RigidCollider.h"
+#include "Chaos/CollisionFilterData.h"
+#include "PhysicalMaterials/PhysicalMaterial.h"
+#include "Physics/Experimental/ChaosInterfaceUtils.h"
+#include "Physics/PhysicsInterfaceTypes.h"
+#include "PhysicsEngine/AggregateGeom.h"
+#include "RigidPhysics/RigidLog.h"
+
+#if UE_RIGIDPHYSICS_API_ENABLED
+
+namespace UE::Physics
+{
+	FRigidShapeInstance::FRigidShapeInstance()
+		: PerShapeData(nullptr)
+	{
+	}
+
+	FRigidShapeInstance::FRigidShapeInstance(Chaos::FPerShapeData* InPerShapeData)
+		: PerShapeData(InPerShapeData)
+	{
+	}
+
+	FRigidShapeInstance::~FRigidShapeInstance()
+	{
+	}
+
+	Chaos::FPerShapeData* FRigidShapeInstance::GetPerShapeData() const
+	{
+		return PerShapeData;
+	}
+
+	Chaos::FConstImplicitObjectPtr FRigidShapeInstance::GetImplicitObject() const
+	{
+		if (PerShapeData != nullptr)
+		{
+			return PerShapeData->GetGeometry();
+		}
+		return Chaos::FConstImplicitObjectPtr();
+	}
+
+}
+
+
+namespace UE::Physics
+{
+	FVector3f MakeSolidBoxInertia(float InMass, const FVector3f& InSize)
+	{
+		FVector3f SizeSq = InSize * InSize;
+		return (InMass / 12.0f) * FVector3f(
+			SizeSq.Y + SizeSq.Z,
+			SizeSq.X + SizeSq.Z,
+			SizeSq.X + SizeSq.Y);
+	}
+
+	FRigidShapeInstancePtr MakeShapeInstanceFromAddParams(const FGeometryAddParams& InParams)
+	{
+		TArray<Chaos::FImplicitObjectPtr> Geoms;
+		Chaos::FShapesArray Shapes;
+		ChaosInterface::CreateGeometry(InParams, Geoms, Shapes);
+
+		if (Shapes.Num() != 1)
+		{
+			UE_LOGFMT(LogRigidPhysics, Warning, "MakeShapeInstanceFromAddParams is intended to be used for single shapes only. Have {0} shapes.", Shapes.Num());
+		}
+
+		if (Shapes.Num() > 0)
+		{
+			return FRigidShapeInstancePtr(Shapes[0].Release());
+		}
+
+		return FRigidShapeInstancePtr();
+	}
+
+	FRigidShapeInstancePtr MakeColliderFromAggGeom(
+		FKAggregateGeom& InAggGeom, 
+		const FVector3d& InScale,
+		ECollisionTraceFlag CollisionTraceFlag,
+		UPhysicalMaterial* InSimpleMaterial,
+		const TArrayView<UPhysicalMaterial*>& InComplexMaterials,
+		const TArrayView<FPhysicalMaterialMaskParams>& InComplexMaterialMasks,
+		const FTransform& InRelativeTransform)
+	{
+		// Not used - filters are set up after
+		FBodyCollisionData BodyCollisionData;
+
+		FGeometryAddParams AddParams;
+		AddParams.bDoubleSided = false;
+		AddParams.CollisionData = BodyCollisionData;
+		AddParams.CollisionTraceType = CollisionTraceFlag;
+		AddParams.Scale = InScale;
+		AddParams.SimpleMaterial = InSimpleMaterial;
+		AddParams.ComplexMaterials = InComplexMaterials;
+		AddParams.ComplexMaterialMasks = InComplexMaterialMasks;
+		AddParams.LocalTransform = InRelativeTransform;
+		AddParams.WorldTransform = FTransform::Identity;
+		AddParams.Geometry = &InAggGeom;
+		//AddParams.TriMeshGeometries = MakeArrayView(TriMeshGeometries);
+
+		return MakeShapeInstanceFromAddParams(AddParams);
+	}
+
+	FRigidShapeInstancePtr MakeBoxShape(const FVector3f& InSize, const FVector3f& InCenter, UPhysicalMaterial* InMaterial)
+	{
+		using namespace Chaos;
+
+		// Single Box
+		FKBoxElem BoxElem = FKBoxElem(InSize.X, InSize.Y, InSize.Z);
+		BoxElem.Center = FVector(InCenter);
+
+		FKAggregateGeom AggGeom;
+		AggGeom.AddElement(BoxElem);
+
+		FRigidShapeInstancePtr Collider = MakeColliderFromAggGeom(
+			AggGeom, 
+			FVector3d::One(),
+			ECollisionTraceFlag::CTF_UseSimpleAsComplex,
+			InMaterial,
+			{},
+			{},
+			FTransform3d::Identity);
+
+		if (Collider.IsValid())
+		{
+			const EFilterFlags FilterFlags = EFilterFlags::SimpleCollision | EFilterFlags::ContactNotify | EFilterFlags::ModifyContacts;
+			const Filter::FShapeFilterData ShapeFilter = Filter::FShapeFilterBuilder::BuildBlockAll(FilterFlags);
+			Collider->GetPerShapeData()->SetQueryEnabled(true);
+			Collider->GetPerShapeData()->SetSimEnabled(true);
+			Collider->GetPerShapeData()->SetShapeFilterData(ShapeFilter);
+		}
+
+		return Collider;
+	}
+}
+
+
+#endif // UE_RIGIDPHYSICS_API_ENABLED
\ No newline at end of file
diff --git a/Engine/Source/Runtime/Experimental/RigidPhysics/Public/RigidPhysics/RigidBody.h b/Engine/Source/Runtime/Experimental/RigidPhysics/Public/RigidPhysics/RigidBody.h
--- a/Engine/Source/Runtime/Experimental/RigidPhysics/Public/RigidPhysics/RigidBody.h
+++ b/Engine/Source/Runtime/Experimental/RigidPhysics/Public/RigidPhysics/RigidBody.h
@@ -8,6 +8,7 @@
 #include "Math/Bounds.h"
 #include "Math/Transform.h"
 #include "RigidPhysics/RigidBodyContainer.h"
+#include "RigidPhysics/RigidCollider.h"
 #include "RigidPhysics/RigidLog.h"
 #include "RigidPhysics/RigidObjectId.h"
 #include "RigidPhysics/RigidObjectPtr.h"
@@ -82,6 +83,15 @@ namespace UE::Physics
 			return false;
 		}
 
+		UE_INTERNAL RIGIDPHYSICS_API virtual int32 GetNumShapes() const = 0;
+		UE_INTERNAL RIGIDPHYSICS_API virtual FConstRigidShapeInstancePtr GetShape(int32 ShapeIndex) const = 0;
+
+		// Add a Collision Shape
+		UE_INTERNAL RIGIDPHYSICS_API virtual void AddShape(const FConstRigidShapeInstancePtr& InShape)
+		{
+			RIGIDPHYSICS_API_UNSUPPORTED();
+		}
+
 		// Set the transform and initialize body state (velocity etc.) at the new location.
 		UE_INTERNAL RIGIDPHYSICS_API virtual void InitTransform(const FTransform3d& InTransform)
 		{
@@ -131,6 +141,17 @@ namespace UE::Physics
 			RIGIDPHYSICS_API_UNSUPPORTED();
 		}
 
+		UE_INTERNAL RIGIDPHYSICS_API virtual FVector3d GetInertia() const
+		{
+			RIGIDPHYSICS_API_UNSUPPORTED();
+			return FVector3d::Zero();
+		}
+
+		UE_INTERNAL RIGIDPHYSICS_API virtual void SetInertia(const FVector3d& InInertia)
+		{
+			RIGIDPHYSICS_API_UNSUPPORTED();
+		}
+
 		UE_INTERNAL RIGIDPHYSICS_API virtual void SetCCDEnabled(bool bInEnabled)
 		{
 			RIGIDPHYSICS_API_UNSUPPORTED();
@@ -289,6 +310,21 @@ namespace UE::Physics
 			return BodyRaw->IsDynamic();
 		}
 
+		int32 GetNumShapes() const
+		{
+			return BodyRaw->GetNumShapes();
+		}
+
+		FConstRigidShapeInstancePtr GetShape(int32 InShapeIndex) const
+		{
+			return BodyRaw->GetShape(InShapeIndex);
+		}
+
+		void AddShape(const FConstRigidShapeInstancePtr& InShape)
+		{
+			BodyRaw->AddShape(InShape);
+		}
+
 		void InitTransform(const FTransform3d& InTransform)
 		{
 			BodyRaw->InitTransform(InTransform);
@@ -313,7 +349,7 @@ namespace UE::Physics
 		{
 			return BodyRaw->GetBounds();
 		}
-		
+
 		double GetMass() const
 		{
 			return BodyRaw->GetMass();
@@ -323,7 +359,17 @@ namespace UE::Physics
 		{
 			BodyRaw->SetMass(InMass);
 		}
-		
+
+		FVector3d GetInertia() const
+		{
+			return BodyRaw->GetInertia();
+		}
+
+		void SetInertia(const FVector3d& InInertia)
+		{
+			BodyRaw->SetInertia(InInertia);
+		}
+
 		void SetCCDEnabled(bool bInEnabled)
 		{
 			BodyRaw->SetCCDEnabled(bInEnabled);
diff --git a/Engine/Source/Runtime/Experimental/RigidPhysics/Public/RigidPhysics/RigidCollider.h b/Engine/Source/Runtime/Experimental/RigidPhysics/Public/RigidPhysics/RigidCollider.h
new file mode 100644
--- /dev/null
+++ b/Engine/Source/Runtime/Experimental/RigidPhysics/Public/RigidPhysics/RigidCollider.h
@@ -0,0 +1,125 @@
+// Copyright Epic Games, Inc. All Rights Reserved.
+#pragma once
+
+#include "RigidPhysics/RigidFwd.h"
+
+#if UE_RIGIDPHYSICS_API_ENABLED
+
+#include "Chaos/ImplicitFwd.h"
+#include "Math/Bounds.h"
+#include "Math/Transform.h"
+#include "RigidPhysics/RigidLog.h"
+#include "RigidPhysics/RigidTyped.h"
+
+class UPhysicalMaterial;
+
+namespace Chaos
+{
+	class FPerShapeData;
+}
+
+namespace UE::Physics
+{
+	// TODO_CHAOS: Should this be an interface with scene-specific implementation?
+	// Not if we want to use this data type in assets/bodyinstance etc. but it would
+	// allow us to use FPerShapeData and FImplicitObject in the ChaosAsync but move
+	// away from them in new implementations.
+	class UE_INTERNAL FRigidShapeInstance
+	{
+	public:
+		UE_INTERNAL RIGIDPHYSICS_API FRigidShapeInstance();
+		UE_INTERNAL RIGIDPHYSICS_API virtual ~FRigidShapeInstance();
+
+		// Existing Chaos implementation exposed for internal use only
+		UE_INTERNAL RIGIDPHYSICS_API FRigidShapeInstance(Chaos::FPerShapeData* InPerShapeData);
+		UE_INTERNAL RIGIDPHYSICS_API Chaos::FPerShapeData* GetPerShapeData() const;
+		UE_INTERNAL RIGIDPHYSICS_API Chaos::FConstImplicitObjectPtr GetImplicitObject() const;
+
+	private:
+		// NOTE: Danger...we hold a pointer to a shape which is the current implementation, 
+		// but its lifetime is currently managed by the Proxy in a UniquePtr.
+		// See FRigidBodyAsyncGT::AddShape
+		Chaos::FPerShapeData* PerShapeData;
+	};
+
+
+	//
+	// A pointer to a collider. Eventually this will just be the equivalent of a raw pointer
+	// to FRigidShapeInstance, but for now we contain the FRigidShapeInstance by value. This is so that
+	// we don't need to store an array of FRigidShapeInstance objects in our FRigidBody class until
+	// we move away from it being a thin wrapper around FPerShapeData (the particle class
+	// already has an array of PerShapeData so we don't want to double up).
+	//
+	template<bool bConstParam>
+	class TRigidShapeInstancePtr
+	{
+	public:
+		static const bool bIsPtrToConst = bConstParam;
+
+		TRigidShapeInstancePtr()
+			: ShapeInstance(nullptr)
+		{
+		}
+
+		TRigidShapeInstancePtr(const TRigidShapeInstancePtr& R)
+			: ShapeInstance(R.Get()->GetPerShapeData())
+		{
+		}
+
+		// Auto-cast non-const to const ptr
+		TRigidShapeInstancePtr(const TRigidShapeInstancePtr<false>& R) requires (bIsPtrToConst)
+			: ShapeInstance(R.Get()->GetPerShapeData())
+		{
+		}
+
+		bool IsValid() const
+		{
+			return ShapeInstance.GetPerShapeData() != nullptr;
+		}
+
+		void Reset()
+		{
+			ShapeInstance = FRigidShapeInstance();
+		}
+
+		operator bool() const
+		{
+			return IsValid();
+		}
+
+		FRigidShapeInstance* operator->() const requires (!bIsPtrToConst)
+		{
+			return &ShapeInstance;
+		}
+
+		const FRigidShapeInstance* operator->() const  requires (bIsPtrToConst)
+		{ 
+			return &ShapeInstance;
+		}
+
+
+		UE_INTERNAL TRigidShapeInstancePtr(Chaos::FPerShapeData* InShapeInstance)
+			: ShapeInstance(InShapeInstance)
+		{
+		}
+
+		UE_INTERNAL const FRigidShapeInstance* Get() const
+		{
+			return &ShapeInstance;
+		}
+
+	private:
+		mutable FRigidShapeInstance ShapeInstance;
+	};
+}
+
+namespace UE::Physics
+{
+	// Utility for Unit Tests until the API if more fully implemented
+	extern UE_INTERNAL RIGIDPHYSICS_API FRigidShapeInstancePtr MakeBoxShape(const FVector3f& InSize, const FVector3f& InCenter = FVector3f::Zero(), UPhysicalMaterial* InMaterial = nullptr);
+	extern UE_INTERNAL RIGIDPHYSICS_API FVector3f MakeSolidBoxInertia(float InMass, const FVector3f& InSize);
+
+}
+
+
+#endif
diff --git a/Engine/Source/Runtime/Experimental/RigidPhysics/Public/RigidPhysics/RigidFwd.h b/Engine/Source/Runtime/Experimental/RigidPhysics/Public/RigidPhysics/RigidFwd.h
--- a/Engine/Source/Runtime/Experimental/RigidPhysics/Public/RigidPhysics/RigidFwd.h
+++ b/Engine/Source/Runtime/Experimental/RigidPhysics/Public/RigidPhysics/RigidFwd.h
@@ -58,6 +58,7 @@ namespace UE::Physics
 	class FRigidQueryDesc;
 	class FRigidQueryHit;
 	class FRigidSceneHandle;
+	class FRigidShapeInstance;
 
 	class IRigidBody;
 	class IRigidBodyContainer;
@@ -71,6 +72,7 @@ namespace UE::Physics
 
 	template<typename ContextType> class TRigidBodyContainerPtrImpl;
 	template<typename ContextType> class TRigidBodyPtrImpl;
+	template<bool bConstParam> class TRigidShapeInstancePtr;
 	template<typename PtrImplType> class TRigidObjectPtr;
 	template<typename ContextType> class TRigidScenePtrImpl;
 
@@ -78,9 +80,10 @@ namespace UE::Physics
 	template<typename ContextType> using TRigidBodyPtr = TRigidObjectPtr<TRigidBodyPtrImpl<ContextType>>;
 	template<typename ContextType> using TRigidScenePtr = TRigidObjectPtr<TRigidScenePtrImpl<ContextType>>;
 
-	using FRigidSceneWeakPtr = TWeakPtr<IRigidScene, ESPMode::ThreadSafe>;
-
+	using FConstRigidShapeInstancePtr = TRigidShapeInstancePtr<true>;
+	using FRigidShapeInstancePtr = TRigidShapeInstancePtr<false>;
 	using FRigidDebugName = Chaos::FSharedDebugName;
+	using FRigidSceneWeakPtr = TWeakPtr<IRigidScene, ESPMode::ThreadSafe>;
 }
 
 namespace UE::Physics
