PR #13954: fix(deps): on-headers is vulnerable to http response header manipulation (and other vulnerability version bumps)

#Synced-CL 50931157

diff --git a/Engine/Plugins/Experimental/Web/WebAPI/WebAPIGeneratorApp/package.json b/Engine/Plugins/Experimental/Web/WebAPI/WebAPIGeneratorApp/package.json
--- a/Engine/Plugins/Experimental/Web/WebAPI/WebAPIGeneratorApp/package.json
+++ b/Engine/Plugins/Experimental/Web/WebAPI/WebAPIGeneratorApp/package.json
@@ -10,7 +10,7 @@
     "dev:nodemon": "tsc && nodemon -w src -w templates -e ts,json,liquid -x ts-node ./src/index.ts",
     "find-deadcode": "ts-prune"
   },
-  "keywords": ["express", "typescript"],
+ "keywords": ["express", "typescript"],
   "devDependencies": {
     "@types/express": "^4.17.13",
     "@types/node": "^14.17.14",
@@ -18,7 +18,7 @@
     "express": "^4.17.1",
     "fs-extra": "^9.0.1",
     "lodash": "^4.17.20",
-    "nodemon": "^2.0.12",
+    "nodemon": "^3.1.11",
     "superagent": "^6.0.0",
     "ts-node": "^9.1.1",
     "typescript": "^4.4.2",
@@ -31,10 +31,10 @@
     "@types/lodash": "^4.14.160",
     "@types/node": "^14.6.0",
     "@types/superagent": "4.1.10",
-    "@types/ws": "^7.2.6",  
+    "@types/ws": "^7.2.6",
     "express": "^4.17.1",
     "http-status": "^1.5.0",
-    "liquidjs": "^9.25.1",
-    "typescript": "^4.0.2"    
+    "liquidjs": "^10.24.0",
+    "typescript": "^4.0.2"
   }
 }
