Chaos API: Adding joint constraint.
- Added a base JointConstraint and JointConstraint6DOF.
- Added a test fixture and moved TestSceneModifier to be shareable.
- Updated visualizers for joints, also fixed some intrinsics to be more generic (need to update remaining ones later).

#Synced-CL 50971911

diff --git a/Engine/Plugins/Experimental/ChaosRigidPhysicsAsync/Source/ChaosRigidPhysicsAsync/Private/ChaosRigidPhysicsAsync/JointConstraint6DOFAsync.cpp b/Engine/Plugins/Experimental/ChaosRigidPhysicsAsync/Source/ChaosRigidPhysicsAsync/Private/ChaosRigidPhysicsAsync/JointConstraint6DOFAsync.cpp
new file mode 100644
--- /dev/null
+++ b/Engine/Plugins/Experimental/ChaosRigidPhysicsAsync/Source/ChaosRigidPhysicsAsync/Private/ChaosRigidPhysicsAsync/JointConstraint6DOFAsync.cpp
@@ -0,0 +1,187 @@
+// Copyright Epic Games, Inc. All Rights Reserved.
+#include "ChaosRigidPhysicsAsync/JointConstraint6DOFAsync.h"
+
+#if UE_RIGIDPHYSICS_API_ENABLED
+
+#include "ChaosRigidPhysicsAsync/RigidSceneAsync.h"
+#include "Chaos/PBDJointConstraintData.h"
+
+namespace Chaos::Rigids::Async
+{
+	Chaos::EJointMotionType Convert(UE::Physics::EJointMotionType InMotionType)
+	{
+		switch (InMotionType)
+		{
+			case UE::Physics::EJointMotionType::Free: return Chaos::EJointMotionType::Free;
+			case UE::Physics::EJointMotionType::Limited: return Chaos::EJointMotionType::Limited;
+			case UE::Physics::EJointMotionType::Locked: return Chaos::EJointMotionType::Locked;
+			default: return Chaos::EJointMotionType::Free;
+		}
+	}
+
+	UE::Physics::EJointMotionType Convert(Chaos::EJointMotionType InMotionType)
+	{
+		switch (InMotionType)
+		{
+			case Chaos::EJointMotionType::Free: return UE::Physics::EJointMotionType::Free;
+			case Chaos::EJointMotionType::Limited: return UE::Physics::EJointMotionType::Limited;
+			case Chaos::EJointMotionType::Locked: return UE::Physics::EJointMotionType::Locked;
+			default: return UE::Physics::EJointMotionType::Free;
+		}
+	}
+
+	UE_RIGIDPHYSICS_RIGIDTYPED_IMPL(FJointConstraint6DOFAsyncGT, UE::Physics::IJointConstraint6DOF);
+
+	FJointConstraint6DOFAsyncGT::FJointConstraint6DOFAsyncGT(FRigidSceneAsyncGT* InScene)
+		: ConstraintId()
+		, Scene(InScene)
+		, Constraint(nullptr)
+	{
+		Constraint = new Chaos::FJointConstraint();
+	}
+
+	FJointConstraint6DOFAsyncGT::~FJointConstraint6DOFAsyncGT()
+	{
+		FPhysicsConstraintHandle PhysicsHandle;
+		PhysicsHandle.Constraint = Constraint;
+		FChaosEngineInterface::ReleaseConstraint(PhysicsHandle);
+		Constraint = nullptr;
+	}
+
+	FJointConstraintHandle FJointConstraint6DOFAsyncGT::GetHandle() const
+	{
+		return FJointConstraintHandle(Scene->GetHandle(), ConstraintId);
+	}
+
+	void FJointConstraint6DOFAsyncGT::Activate()
+	{
+		Scene->ActivateConstraint(this);
+	}
+
+	void FJointConstraint6DOFAsyncGT::SetBodies(IRigidBody* Body0, IRigidBody* Body1)
+	{
+		FRigidBodyAsyncGT* TypedBody0 = Body0->AsA<FRigidBodyAsyncGT>();
+		FRigidBodyAsyncGT* TypedBody1 = Body1->AsA<FRigidBodyAsyncGT>();
+		// TODO: Handle null (attaching to world anchor)
+		if (TypedBody0 && TypedBody1)
+		{
+			Constraint->SetParticleProxies({ TypedBody0->GetProxy(), TypedBody1->GetProxy() });
+		}
+	}
+
+	void FJointConstraint6DOFAsyncGT::SetJointTransforms(const FTransform3d& InTransform0, const FTransform3d& InTransform1)
+	{
+		Constraint->SetJointTransforms({ InTransform0,InTransform1 });
+	}
+
+	double FJointConstraint6DOFAsyncGT::GetLinearLimit() const
+	{
+		return Constraint->GetLinearLimit();
+	}
+
+	void FJointConstraint6DOFAsyncGT::SetLinearLimit(const double InLinearLimit)
+	{
+		Constraint->SetLinearLimit(InLinearLimit);
+	}
+
+	UE::Physics::EJointMotionType FJointConstraint6DOFAsyncGT::GetLinearMotionTypesX() const
+	{
+		return Convert(Constraint->GetLinearMotionTypesX());
+	}
+
+	void FJointConstraint6DOFAsyncGT::SetLinearMotionTypesX(const EJointMotionType InMotionType)
+	{
+		Constraint->SetLinearMotionTypesX(Convert(InMotionType));
+	}
+
+	UE::Physics::EJointMotionType FJointConstraint6DOFAsyncGT::GetLinearMotionTypesY() const
+	{
+		return Convert(Constraint->GetLinearMotionTypesY());
+	}
+
+	void FJointConstraint6DOFAsyncGT::SetLinearMotionTypesY(const EJointMotionType InMotionType)
+	{
+		Constraint->SetLinearMotionTypesY(Convert(InMotionType));
+	}
+
+	UE::Physics::EJointMotionType FJointConstraint6DOFAsyncGT::GetLinearMotionTypesZ() const
+	{
+		return Convert(Constraint->GetLinearMotionTypesZ());
+	}
+
+	void FJointConstraint6DOFAsyncGT::SetLinearMotionTypesZ(const EJointMotionType InMotionType)
+	{
+		Constraint->SetLinearMotionTypesZ(Convert(InMotionType));
+	}
+
+	void FJointConstraint6DOFAsyncGT::SetId(const FRigidObjectId& InId)
+	{
+		ConstraintId = InId;
+	}
+
+	Chaos::FJointConstraint* FJointConstraint6DOFAsyncGT::GetJointConstraint()
+	{
+		return Constraint;
+	}
+
+	const Chaos::FJointConstraint* FJointConstraint6DOFAsyncGT::GetJointConstraint() const
+	{
+		return Constraint;
+	}
+
+	UE_RIGIDPHYSICS_RIGIDTYPED_IMPL(FJointConstraint6DOFAsyncPT, UE::Physics::IJointConstraint6DOF);
+
+	FJointConstraint6DOFAsyncPT::FJointConstraint6DOFAsyncPT(const FJointConstraintHandle& InHandle, FRigidSceneAsyncPT* InScene, Chaos::FJointConstraint* InConstraint)
+		: Handle(InHandle)
+		, Scene(InScene)
+		, Constraint(InConstraint)
+	{
+	}
+
+	FJointConstraintHandle FJointConstraint6DOFAsyncPT::GetHandle() const
+	{
+		return Handle;
+	}
+
+	double FJointConstraint6DOFAsyncPT::GetLinearLimit() const
+	{
+		return Constraint->GetLinearLimit();
+	}
+
+	void FJointConstraint6DOFAsyncPT::SetLinearLimit(const double InLinearLimit)
+	{
+		Constraint->SetLinearLimit(InLinearLimit);
+	}
+
+	UE::Physics::EJointMotionType FJointConstraint6DOFAsyncPT::GetLinearMotionTypesX() const
+	{
+		return Convert(Constraint->GetLinearMotionTypesX());
+	}
+
+	void FJointConstraint6DOFAsyncPT::SetLinearMotionTypesX(const EJointMotionType InMotionType)
+	{
+		Constraint->SetLinearMotionTypesX(Convert(InMotionType));
+	}
+
+	UE::Physics::EJointMotionType FJointConstraint6DOFAsyncPT::GetLinearMotionTypesY() const
+	{
+		return Convert(Constraint->GetLinearMotionTypesY());
+	}
+
+	void FJointConstraint6DOFAsyncPT::SetLinearMotionTypesY(const EJointMotionType InMotionType)
+	{
+		Constraint->SetLinearMotionTypesY(Convert(InMotionType));
+	}
+
+	UE::Physics::EJointMotionType FJointConstraint6DOFAsyncPT::GetLinearMotionTypesZ() const
+	{
+		return Convert(Constraint->GetLinearMotionTypesZ());
+	}
+
+	void FJointConstraint6DOFAsyncPT::SetLinearMotionTypesZ(const EJointMotionType InMotionType)
+	{
+		Constraint->SetLinearMotionTypesZ(Convert(InMotionType));
+	}
+} // namespace Chaos::Rigids::Async
+
+#endif // UE_RIGIDPHYSICS_API_ENABLED
diff --git a/Engine/Plugins/Experimental/ChaosRigidPhysicsAsync/Source/ChaosRigidPhysicsAsync/Private/ChaosRigidPhysicsAsync/RigidSceneAsync.cpp b/Engine/Plugins/Experimental/ChaosRigidPhysicsAsync/Source/ChaosRigidPhysicsAsync/Private/ChaosRigidPhysicsAsync/RigidSceneAsync.cpp
--- a/Engine/Plugins/Experimental/ChaosRigidPhysicsAsync/Source/ChaosRigidPhysicsAsync/Private/ChaosRigidPhysicsAsync/RigidSceneAsync.cpp
+++ b/Engine/Plugins/Experimental/ChaosRigidPhysicsAsync/Source/ChaosRigidPhysicsAsync/Private/ChaosRigidPhysicsAsync/RigidSceneAsync.cpp
@@ -68,6 +68,8 @@ namespace Chaos::Rigids::Async
 
 		UnregisterBodyContainer(StaticBodyPool.Get());
 		StaticBodyPool.Reset();
+
+		FRigidScene::DestructScene();
 	}
 
 	void FRigidSceneAsyncGT::Startup()
@@ -311,6 +313,65 @@ namespace Chaos::Rigids::Async
 			delete GC;
 		}
 	}
+	
+	FJointConstraint6DOFAsyncGT* FRigidSceneAsyncGT::CreateJointConstraint6DOF()
+	{
+		FJointConstraint6DOFAsyncGT* Constraint = new FJointConstraint6DOFAsyncGT(this);
+		const FRigidObjectId Id = GetJointConstraintRegistry()->AddObject(Constraint);
+		Constraint->SetId(Id);
+
+		if (FPBDRigidsSolver* Solver = GetSolver())
+		{
+			// Note: These need to be cached outside the command as the GT object can get destroyed before the command is executed.
+			const FJointConstraint6DOFHandle Handle = Constraint->GetTypedHandle();
+			Chaos::FJointConstraint* ChaosJointConstraint = Constraint->GetJointConstraint();
+			Solver->EnqueueCommandImmediate([this, Handle, ChaosJointConstraint]()
+				{
+					ScenePT->CreateJointConstraint6DOFPT(Handle, ChaosJointConstraint);
+				});
+		}
+		return Constraint;
+	}
+
+	void FRigidSceneAsyncGT::DestroyJointConstraint(IJointConstraint* InConstraint)
+	{
+		if (FJointConstraint6DOFAsyncGT* Constraint6DOF = InConstraint->AsAChecked<FJointConstraint6DOFAsyncGT>())
+		{
+			const FJointConstraint6DOFHandle Handle = Constraint6DOF->GetTypedHandle();
+			if (FPBDRigidsSolver* Solver = GetSolver())
+			{
+				Solver->EnqueueCommandImmediate(
+					[this, Handle]()
+					{
+						ScenePT->DestroyJointConstraintPT(Handle);
+					});
+			}
+
+			GetJointConstraintRegistry()->RemoveObject(Handle.GetId());
+			delete Constraint6DOF;
+		}
+	}
+
+	IJointConstraint* FRigidSceneAsyncGT::GetJointConstraint(const FRigidObjectId& InConstraintId) const
+	{
+		return GetJointConstraintRegistry()->FindObject(InConstraintId);
+	}
+
+	void FRigidSceneAsyncGT::ActivateConstraint(FJointConstraint6DOFAsyncGT* InConstraint)
+	{
+		if (FPBDRigidsSolver* Solver = GetSolver())
+		{
+			Solver->EnqueueCommandImmediate([this, Solver, InConstraint]()
+				{
+					Solver->RegisterObject(InConstraint->GetJointConstraint());
+				});
+		}
+	}
+
+	void FRigidSceneAsyncGT::DeactivateConstraint(FJointConstraint6DOFAsyncGT* InConstraint)
+	{
+		// TODO
+	}
 
 	void FRigidSceneAsyncGT::RegisterModifier(UE::Physics::IRigidSceneModifier* Modifier)
 	{
@@ -643,6 +704,46 @@ namespace Chaos::Rigids::Async
 
 		return nullptr;
 	}
+	
+	FJointConstraint6DOFAsyncPT* FRigidSceneAsyncPT::CreateJointConstraint6DOFPT(const FJointConstraint6DOFHandle& InHandle, Chaos::FJointConstraint* InJointConstraint)
+	{
+		CheckSimContext();
+
+		FJointConstraint6DOFAsyncPT* Result = new FJointConstraint6DOFAsyncPT(InHandle, this, InJointConstraint);
+		JointConstraintRegistry.Add(InHandle.GetId(), Result);
+		return Result;
+	}
+
+	void FRigidSceneAsyncPT::DestroyJointConstraintPT(const FJointConstraint6DOFHandle& InHandle)
+	{
+		if (IJointConstraint** JointConstraintPtr = JointConstraintRegistry.Find(InHandle.GetId()))
+		{
+			IJointConstraint* JointConstraint = *JointConstraintPtr;
+			JointConstraintRegistry.Remove(InHandle.GetId());
+			delete JointConstraint;
+		}
+	}
+
+	IJointConstraint* FRigidSceneAsyncPT::GetJointConstraint(const FRigidObjectId& InConstraintId) const
+	{
+		CheckSimContext();
+
+		if (IJointConstraint*const* JointPtr = JointConstraintRegistry.Find(InConstraintId))
+		{
+			return *JointPtr;
+		}
+		return nullptr;
+	}
+
+	void FRigidSceneAsyncPT::ActiveJointConstraint(const FJointConstraint6DOFHandle& InHandle)
+	{
+		CheckSimContext();
+	}
+
+	void FRigidSceneAsyncPT::DeactivateJointConstraint(const FJointConstraint6DOFHandle& InHandle)
+	{
+		CheckSimContext();
+	}
 
 	void FRigidSceneAsyncPT::CheckGameContext() const
 	{
diff --git a/Engine/Plugins/Experimental/ChaosRigidPhysicsAsync/Source/ChaosRigidPhysicsAsync/Public/ChaosRigidPhysicsAsync/JointConstraint6DOFAsync.h b/Engine/Plugins/Experimental/ChaosRigidPhysicsAsync/Source/ChaosRigidPhysicsAsync/Public/ChaosRigidPhysicsAsync/JointConstraint6DOFAsync.h
new file mode 100644
--- /dev/null
+++ b/Engine/Plugins/Experimental/ChaosRigidPhysicsAsync/Source/ChaosRigidPhysicsAsync/Public/ChaosRigidPhysicsAsync/JointConstraint6DOFAsync.h
@@ -0,0 +1,87 @@
+// Copyright Epic Games, Inc. All Rights Reserved.
+#pragma once 
+
+#include "ChaosRigidPhysicsAsync/RigidFwdAsync.h"
+
+#if UE_RIGIDPHYSICS_API_ENABLED
+
+#include "RigidPhysics/JointConstraint6DOF.h"
+
+namespace Chaos
+{
+	class FJointConstraint;
+} // namespace Chaos
+
+namespace Chaos::Rigids::Async
+{
+	class UE_INTERNAL FJointConstraint6DOFAsyncGT : public UE::Physics::IJointConstraint6DOF
+	{
+		using EJointMotionType = UE::Physics::EJointMotionType;
+		using IRigidBody = UE::Physics::IRigidBody;
+
+	public:
+		UE_RIGIDPHYSICS_RIGIDTYPED_DECL(CHAOSRIGIDPHYSICSASYNC_API, FJointConstraint6DOFAsyncGT);
+
+		UE_INTERNAL CHAOSRIGIDPHYSICSASYNC_API FJointConstraint6DOFAsyncGT(FRigidSceneAsyncGT* InScene);
+		UE_INTERNAL virtual ~FJointConstraint6DOFAsyncGT();
+
+		UE_INTERNAL virtual FJointConstraintHandle GetHandle() const override final;
+
+		UE_INTERNAL virtual void Activate() override final;
+		UE_INTERNAL virtual void SetBodies(IRigidBody* Body0, IRigidBody* Body1) override final;
+		UE_INTERNAL virtual void SetJointTransforms(const FTransform3d& InTransform0, const FTransform3d& InTransform1) override final;
+
+		UE_INTERNAL virtual double GetLinearLimit() const override final;
+		UE_INTERNAL virtual void SetLinearLimit(const double InLinearLimit) override final;
+
+		UE_INTERNAL virtual EJointMotionType GetLinearMotionTypesX() const override final;
+		UE_INTERNAL virtual void SetLinearMotionTypesX(const EJointMotionType InMotionType) override final;
+
+		UE_INTERNAL virtual EJointMotionType GetLinearMotionTypesY() const override final;
+		UE_INTERNAL virtual void SetLinearMotionTypesY(const EJointMotionType InMotionType) override final;
+
+		UE_INTERNAL virtual EJointMotionType GetLinearMotionTypesZ() const override final;
+		UE_INTERNAL virtual void SetLinearMotionTypesZ(const EJointMotionType InMotionType) override final;
+
+		UE_INTERNAL void SetId(const FRigidObjectId& InId);
+		UE_INTERNAL Chaos::FJointConstraint* GetJointConstraint();
+		UE_INTERNAL const Chaos::FJointConstraint* GetJointConstraint() const;
+
+	private:
+		UE::Physics::FRigidObjectId ConstraintId;
+		FRigidSceneAsyncGT* Scene = nullptr;
+		Chaos::FJointConstraint* Constraint = nullptr;
+	};
+
+	class UE_INTERNAL FJointConstraint6DOFAsyncPT : public UE::Physics::IJointConstraint6DOF
+	{
+		using EJointMotionType = UE::Physics::EJointMotionType;
+
+	public:
+		UE_RIGIDPHYSICS_RIGIDTYPED_DECL(CHAOSRIGIDPHYSICSASYNC_API, FJointConstraint6DOFAsyncPT);
+
+		UE_INTERNAL CHAOSRIGIDPHYSICSASYNC_API FJointConstraint6DOFAsyncPT(const FJointConstraintHandle& InHandle, FRigidSceneAsyncPT* InScene, Chaos::FJointConstraint* InConstraint);
+		UE_INTERNAL virtual ~FJointConstraint6DOFAsyncPT() = default;
+
+		UE_INTERNAL virtual FJointConstraintHandle GetHandle() const override final;
+
+		UE_INTERNAL virtual double GetLinearLimit() const override final;
+		UE_INTERNAL virtual void SetLinearLimit(const double InLinearLimit) override final;
+
+		UE_INTERNAL virtual EJointMotionType GetLinearMotionTypesX() const override final;
+		UE_INTERNAL virtual void SetLinearMotionTypesX(const EJointMotionType InMotionType) override final;
+
+		UE_INTERNAL virtual EJointMotionType GetLinearMotionTypesY() const override final;
+		UE_INTERNAL virtual void SetLinearMotionTypesY(const EJointMotionType InMotionType) override final;
+
+		UE_INTERNAL virtual EJointMotionType GetLinearMotionTypesZ() const override final;
+		UE_INTERNAL virtual void SetLinearMotionTypesZ(const EJointMotionType InMotionType) override final;
+
+	private:
+		FJointConstraintHandle Handle;
+		FRigidSceneAsyncPT* Scene = nullptr;
+		Chaos::FJointConstraint* Constraint = nullptr;
+	};
+} // namespace Chaos::Rigids::Async
+
+#endif // UE_RIGIDPHYSICS_API_ENABLED
diff --git a/Engine/Plugins/Experimental/ChaosRigidPhysicsAsync/Source/ChaosRigidPhysicsAsync/Public/ChaosRigidPhysicsAsync/RigidFwdAsync.h b/Engine/Plugins/Experimental/ChaosRigidPhysicsAsync/Source/ChaosRigidPhysicsAsync/Public/ChaosRigidPhysicsAsync/RigidFwdAsync.h
--- a/Engine/Plugins/Experimental/ChaosRigidPhysicsAsync/Source/ChaosRigidPhysicsAsync/Public/ChaosRigidPhysicsAsync/RigidFwdAsync.h
+++ b/Engine/Plugins/Experimental/ChaosRigidPhysicsAsync/Source/ChaosRigidPhysicsAsync/Public/ChaosRigidPhysicsAsync/RigidFwdAsync.h
@@ -13,6 +13,8 @@ namespace Chaos::Rigids::Async
 // Internal fwd decls
 namespace Chaos::Rigids::Async
 {
+	class UE_INTERNAL FJointConstraint6DOFAsyncGT;
+	class UE_INTERNAL FJointConstraint6DOFAsyncPT;
 	class UE_INTERNAL FRigidBodyAsyncGT;
 	class UE_INTERNAL FRigidBodyAsyncPT;
 	class UE_INTERNAL FRigidBodyPoolAsyncGT;
@@ -21,6 +23,13 @@ namespace Chaos::Rigids::Async
 	class UE_INTERNAL FRigidGeometryCollectionAsyncPT;
 	class UE_INTERNAL FRigidSceneAsyncGT;
 	class UE_INTERNAL FRigidSceneAsyncPT;
+	
+	using UE_INTERNAL FJointConstraintHandle = UE::Physics::FJointConstraintHandle;
+	using UE_INTERNAL FJointConstraint6DOFHandle = UE::Physics::FJointConstraint6DOFHandle;
+	using UE_INTERNAL FJointConstraintRegistry = UE::Physics::FJointConstraintRegistry;
+	using UE_INTERNAL FRigidObjectId = UE::Physics::FRigidObjectId;
+	using UE_INTERNAL IJointConstraint = UE::Physics::IJointConstraint;
+	using UE_INTERNAL IJointConstraint6DOF = UE::Physics::IJointConstraint6DOF;
 }
 
-#endif // UE_RIGIDPHYSICS_API_ENABLED
\ No newline at end of file
+#endif // UE_RIGIDPHYSICS_API_ENABLED
diff --git a/Engine/Plugins/Experimental/ChaosRigidPhysicsAsync/Source/ChaosRigidPhysicsAsync/Public/ChaosRigidPhysicsAsync/RigidSceneAsync.h b/Engine/Plugins/Experimental/ChaosRigidPhysicsAsync/Source/ChaosRigidPhysicsAsync/Public/ChaosRigidPhysicsAsync/RigidSceneAsync.h
--- a/Engine/Plugins/Experimental/ChaosRigidPhysicsAsync/Source/ChaosRigidPhysicsAsync/Public/ChaosRigidPhysicsAsync/RigidSceneAsync.h
+++ b/Engine/Plugins/Experimental/ChaosRigidPhysicsAsync/Source/ChaosRigidPhysicsAsync/Public/ChaosRigidPhysicsAsync/RigidSceneAsync.h
@@ -6,6 +6,7 @@
 #if UE_RIGIDPHYSICS_API_ENABLED
 
 #include "Chaos/ObjectPool.h"
+#include "ChaosRigidPhysicsAsync/JointConstraint6DOFAsync.h"
 #include "ChaosRigidPhysicsAsync/RigidBodyAsync.h"
 #include "ChaosRigidPhysicsAsync/RigidBodyContainerAsync.h"
 #include "ChaosRigidPhysicsAsync/RigidFactoryAsync.h"
@@ -77,6 +78,13 @@ namespace Chaos::Rigids::Async
 		//
 		// IRIGIDSCENE END
 		//
+		
+		// Joint Constraint
+		UE_INTERNAL virtual FJointConstraint6DOFAsyncGT* CreateJointConstraint6DOF() override final;
+		UE_INTERNAL virtual void DestroyJointConstraint(IJointConstraint* InConstraint) override final;
+		UE_INTERNAL virtual IJointConstraint* GetJointConstraint(const FRigidObjectId& InConstraintId) const override final;
+		UE_INTERNAL CHAOSRIGIDPHYSICSASYNC_API void ActivateConstraint(FJointConstraint6DOFAsyncGT* InConstraint);
+		UE_INTERNAL CHAOSRIGIDPHYSICSASYNC_API void DeactivateConstraint(FJointConstraint6DOFAsyncGT* InConstraint);
 
 	public:
 		UE_INTERNAL CHAOSRIGIDPHYSICSASYNC_API void Startup();
@@ -166,6 +174,13 @@ namespace Chaos::Rigids::Async
 		UE_INTERNAL CHAOSRIGIDPHYSICSASYNC_API void DestroyBodyPT(const UE::Physics::FRigidBodyHandle& InBodyHandle);
 		UE_INTERNAL CHAOSRIGIDPHYSICSASYNC_API void ActivateBody(const UE::Physics::FRigidBodyHandle& InBodyHandle);
 		UE_INTERNAL CHAOSRIGIDPHYSICSASYNC_API void DeactivateBody(const UE::Physics::FRigidBodyHandle& InBodyHandle);
+		
+		// Joint Constraint
+		UE_INTERNAL FJointConstraint6DOFAsyncPT* CreateJointConstraint6DOFPT(const FJointConstraint6DOFHandle& InHandle, Chaos::FJointConstraint* InJointConstraint);
+		UE_INTERNAL void DestroyJointConstraintPT(const FJointConstraint6DOFHandle& InHandle);
+		UE_INTERNAL virtual IJointConstraint* GetJointConstraint(const FRigidObjectId& InConstraintId) const override final;
+		UE_INTERNAL void ActiveJointConstraint(const FJointConstraint6DOFHandle& InHandle);
+		UE_INTERNAL void DeactivateJointConstraint(const FJointConstraint6DOFHandle& InHandle);
 
 	private:
 		void CheckGameContext() const;
@@ -180,6 +195,7 @@ namespace Chaos::Rigids::Async
 		TUniquePtr<FRigidBodyPoolAsyncPT> StaticBodyPool;
 		TUniquePtr<FRigidBodyPoolAsyncPT> DynamicBodyPool;
 		TMap<UE::Physics::FRigidObjectId, UE::Physics::IRigidBodyContainer*> BodyContainerRegistry;
+		TMap<UE::Physics::FRigidObjectId, UE::Physics::IJointConstraint*> JointConstraintRegistry;
 		TArray<UE::Physics::IRigidSceneModifier*> SceneModifiers;
 		bool bIsStarted;
 	};
diff --git a/Engine/Plugins/Experimental/ChaosRigidPhysicsAsync/Tests/Private/JointConstraintTests.cpp b/Engine/Plugins/Experimental/ChaosRigidPhysicsAsync/Tests/Private/JointConstraintTests.cpp
new file mode 100644
--- /dev/null
+++ b/Engine/Plugins/Experimental/ChaosRigidPhysicsAsync/Tests/Private/JointConstraintTests.cpp
@@ -0,0 +1,196 @@
+// Copyright Epic Games, Inc. All Rights Reserved.
+
+#include "catch2/generators/catch_generators.hpp"
+#include "Chaos/LowLevelTest/ChaosTestHarness.h"
+#include "Chaos/LowLevelTest/ChaosTestScene.h"
+
+#if UE_RIGIDPHYSICS_API_ENABLED
+
+#include "ChaosRigidPhysicsAsyncTest.h"
+#include "RigidPhysics/RigidBody.h"
+#include "RigidPhysics/RigidCollider.h"
+#include "RigidPhysics/RigidContext.h"
+#include "RigidPhysics/RigidScene.h"
+#include "RigidPhysics/JointConstraint6DOF.h"
+#include "RigidTestFixture.h"
+#include "RigidTestUtils.h"
+
+namespace Chaos::LowLevelTest
+{
+	using namespace Chaos;
+	using namespace Chaos::Rigids::Async;
+	using namespace UE::Physics;
+	
+	struct FJointFixture : public FRigidTestFixture
+	{
+		FRigidBodyHandle CreateSimpleBox(FRigidContextGameRW& Context, ERigidMovementType MovementType, const FVector& Location, bool bAutoCleanup = true, const FString& DebugName = "Body")
+		{
+			TRigidBodyPtr<FRigidContextGameRW> BodyPtr = Context->CreateBody(FRigidDebugName(DebugName), MovementType);
+			REQUIRE(BodyPtr);
+			BodyPtr->AddShape(MakeBoxShape(FVector3f(100, 100, 100)));
+			BodyPtr->InitTransform(FTransform(Location));
+			if (MovementType == ERigidMovementType::Dynamic)
+			{
+				BodyPtr->SetIsSleeping(false);
+				BodyPtr->SetMass(1000);
+				BodyPtr->SetInertia(FVector3d(MakeSolidBoxInertia(1000, FVector3f(100, 100, 100))));
+			}
+			BodyPtr->Activate();
+			if (bAutoCleanup)
+			{
+				AutoCleanup(BodyPtr);
+			}
+			return BodyPtr;
+		}
+	};
+
+	TEST_CASE("JointConstraint6DOF: Construct Destruct", "[Chaos][API][JointConstraint][unit]")
+	{
+		FJointFixture Fixture;
+
+		FJointConstraint6DOFHandle JointHandle;
+		if (FRigidContextGameRW Context = Fixture.SceneHandle.LockRW())
+		{
+			TJointConstraint6DOFPtr<FRigidContextGameRW> Constraint = Context->CreateJointConstraint6DOF();
+			REQUIRE(Constraint);
+			JointHandle = Constraint;
+		}
+
+		if (FRigidContextGameRO Context = Fixture.SceneHandle.LockRO())
+		{
+			TJointConstraint6DOFPtr<FRigidContextGameRO> Constraint = JointHandle.Pin(Context);
+			CHECK(Constraint);
+		}
+
+		if (FRigidContextGameRW Context = Fixture.SceneHandle.LockRW())
+		{
+			TJointConstraint6DOFPtr<FRigidContextGameRW> Constraint = JointHandle.Pin(Context);
+			CHECK(Constraint);
+			Context->DestroyJointConstraint(MoveTemp(Constraint));
+
+			CHECK(!JointHandle.Pin(Context));
+		}
+	}
+	
+	TEST_CASE("JointConstraint6DOF: PT Pin", "[Chaos][API][JointConstraint][unit]")
+	{
+		FJointFixture Fixture;
+
+		FJointConstraint6DOFHandle JointHandle;
+		if (FRigidContextGameRW Context = Fixture.SceneHandle.LockRW())
+		{
+			TJointConstraint6DOFPtr<FRigidContextGameRW> Constraint = Context->CreateJointConstraint6DOF();
+			Constraint->SetLinearLimit(123);
+			JointHandle = Constraint;
+			Fixture.AutoCleanup(JointHandle);
+		}
+		
+		Fixture.RunPTCallback([&JointHandle](const FRigidContextSimRW& Context)
+			{
+				TJointConstraint6DOFPtr<FRigidContextSimRW> Constraint = JointHandle.Pin(Context);
+				REQUIRE(Constraint);
+
+				CHECK(Constraint->GetLinearLimit() == 123);
+			});
+	}
+	
+	TEST_CASE("JointConstraint6DOF: Misc Properties", "[Chaos][API][JointConstraint][unit]")
+	{
+		FJointFixture Fixture;
+
+		FJointConstraint6DOFHandle JointHandle;
+		if (FRigidContextGameRW Context = Fixture.SceneHandle.LockRW())
+		{
+			JointHandle = Context->CreateJointConstraint6DOF();
+			Fixture.AutoCleanup(JointHandle);
+		}
+
+		if (FRigidContextGameRW Context = Fixture.SceneHandle.LockRW())
+		{
+			TJointConstraint6DOFPtr<FRigidContextGameRW> Constraint = JointHandle.Pin(Context);
+			SECTION("LinearLimit")
+			{
+				Constraint->SetLinearLimit(150);
+				CHECK(Constraint->GetLinearLimit() == 150);
+			}
+			SECTION("LinearMotionTypes")
+			{
+				const EJointMotionType MotionType = GENERATE(EJointMotionType::Free, EJointMotionType::Limited, EJointMotionType::Locked);
+
+				Constraint->SetLinearMotionTypesX(MotionType);
+				CHECK(Constraint->GetLinearMotionTypesX() == MotionType);
+				Constraint->SetLinearMotionTypesY(MotionType);
+				CHECK(Constraint->GetLinearMotionTypesY() == MotionType);
+				Constraint->SetLinearMotionTypesZ(MotionType);
+				CHECK(Constraint->GetLinearMotionTypesZ() == MotionType);
+			}
+		}
+	}
+
+	TEST_CASE("JointConstraint6DOF: Up/Down Casting", "[Chaos][API][JointConstraint][unit]")
+	{
+		FJointFixture Fixture;
+
+		FJointConstraint6DOFHandle JointHandle6DOF;
+		if (FRigidContextGameRW Context = Fixture.SceneHandle.LockRW())
+		{
+			JointHandle6DOF = Context->CreateJointConstraint6DOF();
+			Fixture.AutoCleanup(JointHandle6DOF);
+		}		
+
+		if (FRigidContextGameRO Context = Fixture.SceneHandle.LockRO())
+		{
+			// Handle Up-cast
+			FJointConstraintHandle JointHandle = JointHandle6DOF;
+			REQUIRE(JointHandle.Pin(Context));
+			// No Handle downcast as that requires pinning
+
+			// Ptr casts:
+			// valid ptr downcast
+			CHECK(TJointConstraint6DOFPtr<FRigidContextGameRO>(JointHandle.Pin(Context)));
+			// ptr upcast not yet supported
+			// CHECK(TJointConstraintPtr<FRigidContextGameRO>(JointHandle6DOF.Pin(Context)));
+			// No way to do invalid downcast currently
+		}
+	}
+	
+	TEST_CASE("JointConstraint6DOF: Simulate Basic", "[Chaos][API][JointConstraint][unit]")
+	{
+		FJointFixture Fixture;
+
+		FJointConstraint6DOFHandle JointHandle;
+		FRigidBodyHandle Body0Handle;
+		FRigidBodyHandle Body1Handle;
+		if (FRigidContextGameRW Context = Fixture.SceneHandle.LockRW())
+		{
+			Body0Handle = Fixture.CreateSimpleBox(Context, ERigidMovementType::Static, FVector::Zero());
+			Body1Handle = Fixture.CreateSimpleBox(Context, ERigidMovementType::Dynamic, FVector(0, 0, -150));
+			TJointConstraint6DOFPtr<FRigidContextGameRW> Constraint = Context->CreateJointConstraint6DOF();
+			REQUIRE(Constraint);
+			JointHandle = Constraint;
+			Fixture.AutoCleanup(JointHandle);
+
+			Constraint->SetBodies(Body0Handle.Pin(Context), Body1Handle.Pin(Context));
+			FTransform Local1, Local2;
+			Constraint->SetJointTransforms(Local1, Local2);
+			Constraint->SetLinearLimit(150);
+			Constraint->SetLinearMotionTypesX(EJointMotionType::Limited);
+			Constraint->SetLinearMotionTypesY(EJointMotionType::Limited);
+			Constraint->SetLinearMotionTypesZ(EJointMotionType::Limited);
+			Constraint->Activate();
+		}
+
+		if (FRigidContextGameRW Context = Fixture.SceneHandle.LockRW())
+		{
+			TRigidBodyPtr<FRigidContextGameRW> Body1Ptr = Body1Handle.Pin(Context);
+			for (int32 TickIndex = 0; TickIndex < 10; ++TickIndex)
+			{
+				Context->Tick(0.01);
+			}
+
+			CHECK(Body1Ptr->GetTransform().GetLocation().Z > -151);
+		}
+	}
+} // namespace Chaos::LowLevelTest
+
+#endif // UE_RIGIDPHYSICS_API_ENABLED
diff --git a/Engine/Plugins/Experimental/ChaosRigidPhysicsAsync/Tests/Private/RigidTestFixture.cpp b/Engine/Plugins/Experimental/ChaosRigidPhysicsAsync/Tests/Private/RigidTestFixture.cpp
new file mode 100644
--- /dev/null
+++ b/Engine/Plugins/Experimental/ChaosRigidPhysicsAsync/Tests/Private/RigidTestFixture.cpp
@@ -0,0 +1,83 @@
+// Copyright Epic Games, Inc. All Rights Reserved.
+#include "RigidTestFixture.h"
+
+#include "Chaos/LowLevelTest/ChaosTestHarness.h"
+#include "Chaos/LowLevelTest/ChaosTestScene.h"
+
+#if UE_RIGIDPHYSICS_API_ENABLED
+
+#include "ChaosRigidPhysicsAsync/RigidFactoryAsync.h"
+#include "ChaosRigidPhysicsAsync/RigidBodyAsync.h"
+#include "ChaosRigidPhysicsAsync/RigidSceneAsync.h"
+#include "ChaosRigidPhysicsAsync/RigidSceneSettingsAsync.h"
+
+#include "RigidPhysics/RigidBody.h"
+#include "RigidPhysics/RigidCollider.h"
+#include "RigidPhysics/RigidContext.h"
+#include "TestSceneModifier.h"
+#include "RigidTestUtils.h"
+
+namespace Chaos::LowLevelTest
+{
+	using namespace Chaos;
+	using namespace Chaos::Rigids::Async;
+	using namespace UE::Physics;
+	
+	FRigidTestFixture::FRigidTestFixture()
+	{
+		SceneHandle = Factory.CreateScene(FRigidDebugName::Make(TEXT("Main")), &SceneSettings);
+	}
+
+	FRigidTestFixture::~FRigidTestFixture()
+	{
+		if (FRigidContextGameRW Context = SceneHandle.LockRW())
+		{
+			for (FJointConstraintHandle& JointHandle : JointHandles)
+			{
+				Context->DestroyJointConstraint(JointHandle.Pin(Context));
+			}
+			for (FRigidBodyHandle& BodyHandle : BodyHandles)
+			{
+				Context->DestroyBody(BodyHandle.Pin(Context));
+			}
+			for (FRigidGeometryCollectionHandle& GCHandle : GeometryCollectionHandles)
+			{
+				Context->DestroyGeometryCollection(GCHandle.Pin(Context));
+			}
+		}
+		Factory.DestroyScene(SceneHandle);
+	}
+
+	void FRigidTestFixture::AutoCleanup(FJointConstraintHandle Handle)
+	{
+		JointHandles.Add(Handle);
+	}
+
+	void FRigidTestFixture::AutoCleanup(FRigidBodyHandle Handle)
+	{
+		BodyHandles.Add(Handle);
+	}
+
+	void FRigidTestFixture::AutoCleanup(FRigidGeometryCollectionHandle Handle)
+	{
+		GeometryCollectionHandles.Add(Handle);
+	}
+
+	void FRigidTestFixture::RunPTCallback(TFunction<void(const FRigidContextSimRW&)> Callback, float Dt, int32 Iterations)
+	{
+		FTestRigidSceneModifier SceneModifier;
+		SceneModifier.PreTickFunc = Callback;
+
+		if (FRigidContextGameRW Context = SceneHandle.LockRW())
+		{
+			Context->RegisterModifier(&SceneModifier);
+			for (int32 I = 0; I < Iterations; ++I)
+			{
+				Context->Tick(Dt);
+			}
+			Context->UnregisterModifier(&SceneModifier);
+		}
+	}
+} // namespace Chaos::LowLevelTest
+
+#endif
diff --git a/Engine/Plugins/Experimental/ChaosRigidPhysicsAsync/Tests/Private/RigidTestFixture.h b/Engine/Plugins/Experimental/ChaosRigidPhysicsAsync/Tests/Private/RigidTestFixture.h
new file mode 100644
--- /dev/null
+++ b/Engine/Plugins/Experimental/ChaosRigidPhysicsAsync/Tests/Private/RigidTestFixture.h
@@ -0,0 +1,38 @@
+// Copyright Epic Games, Inc. All Rights Reserved.
+#pragma once
+
+#include "RigidPhysics/RigidFwd.h"
+#include "ChaosRigidPhysicsAsync/RigidFactoryAsync.h"
+#include "ChaosRigidPhysicsAsync/RigidSceneSettingsAsync.h"
+#include "RigidPhysics/RigidScene.h"
+
+namespace Chaos::LowLevelTest
+{	
+	using FRigidDebugName = UE::Physics::FRigidDebugName;
+	using FRigidContextGameRW = UE::Physics::FRigidContextGameRW;
+	using FRigidContextSimRW = UE::Physics::FRigidContextSimRW;
+	using ERigidMovementType = UE::Physics::ERigidMovementType;
+	using FJointConstraintHandle = UE::Physics::FJointConstraintHandle;
+	using FRigidBodyHandle = UE::Physics::FRigidBodyHandle;
+	using FRigidGeometryCollectionHandle = UE::Physics::FRigidGeometryCollectionHandle;
+
+	struct FRigidTestFixture
+	{
+		FRigidTestFixture();
+		~FRigidTestFixture();
+
+		void AutoCleanup(FJointConstraintHandle Handle);
+		void AutoCleanup(FRigidBodyHandle Handle);
+		void AutoCleanup(FRigidGeometryCollectionHandle Handle);
+
+		void RunPTCallback(TFunction<void(const FRigidContextSimRW&)> Callback, float Dt = 0.01, int32 Iterations = 1);
+
+		Chaos::Rigids::Async::FRigidFactoryAsync Factory;
+		Chaos::Rigids::Async::FRigidSceneSettingsAsync SceneSettings;
+		UE::Physics::FRigidSceneHandle SceneHandle;
+
+		TArray<FRigidBodyHandle> BodyHandles;
+		TArray<FRigidGeometryCollectionHandle> GeometryCollectionHandles;
+		TArray<FJointConstraintHandle> JointHandles;
+	};
+} // namespace Chaos::LowLevelTest
\ No newline at end of file
diff --git a/Engine/Plugins/Experimental/ChaosRigidPhysicsAsync/Tests/Private/SceneTests.cpp b/Engine/Plugins/Experimental/ChaosRigidPhysicsAsync/Tests/Private/SceneTests.cpp
--- a/Engine/Plugins/Experimental/ChaosRigidPhysicsAsync/Tests/Private/SceneTests.cpp
+++ b/Engine/Plugins/Experimental/ChaosRigidPhysicsAsync/Tests/Private/SceneTests.cpp
@@ -14,6 +14,7 @@
 #include "RigidPhysics/RigidContext.h"
 #include "RigidPhysics/RigidModifier.h"
 #include "RigidPhysics/RigidScene.h"
+#include "TestSceneModifier.h"
 
 #include "RigidTestUtils.h"
 
@@ -243,29 +244,6 @@ namespace Chaos::LowLevelTest
 		SceneFactory.DestroyScene(SceneHandle);
 	} // TEST_CASE
 
-	class FTestRigidSceneModifier : public IRigidSceneModifier
-	{
-	public:
-		virtual void PreSimulate(const FRigidContextGameRW& Context) override final
-		{
-			if (PreSimFunc)
-			{
-				PreSimFunc(Context);
-			}
-		}
-
-		virtual void PreTick(const FRigidContextSimRW& Context) override final
-		{
-			if (PreTickFunc)
-			{
-				PreTickFunc(Context);
-			}
-		}
-
-		TFunction<void(const FRigidContextGameRW&)> PreSimFunc;
-		TFunction<void(const FRigidContextSimRW&)> PreTickFunc;
-	};
-
 	TEST_CASE("CallbacksAreCalled", "[Chaos][API][Scene][unit]")
 	{
 		float Dt = 1.0f / 60.0f;
diff --git a/Engine/Plugins/Experimental/ChaosRigidPhysicsAsync/Tests/Private/TestSceneModifier.h b/Engine/Plugins/Experimental/ChaosRigidPhysicsAsync/Tests/Private/TestSceneModifier.h
new file mode 100644
--- /dev/null
+++ b/Engine/Plugins/Experimental/ChaosRigidPhysicsAsync/Tests/Private/TestSceneModifier.h
@@ -0,0 +1,30 @@
+// Copyright Epic Games, Inc. All Rights Reserved.
+#pragma once
+
+#include "RigidPhysics/RigidModifier.h"
+
+namespace Chaos::LowLevelTest
+{
+	class FTestRigidSceneModifier : public UE::Physics::IRigidSceneModifier
+	{
+	public:
+		virtual void PreSimulate(const UE::Physics::FRigidContextGameRW& Context) override final
+		{
+			if (PreSimFunc)
+			{
+				PreSimFunc(Context);
+			}
+		}
+
+		virtual void PreTick(const UE::Physics::FRigidContextSimRW& Context) override final
+		{
+			if (PreTickFunc)
+			{
+				PreTickFunc(Context);
+			}
+		}
+
+		TFunction<void(const UE::Physics::FRigidContextGameRW&)> PreSimFunc;
+		TFunction<void(const UE::Physics::FRigidContextSimRW&)> PreTickFunc;
+	};
+} // namespace Chaos::LowLevelTest
\ No newline at end of file
diff --git a/Engine/Source/Runtime/Experimental/RigidPhysics/Private/RigidPhysics/Internal/IJointConstraint.cpp b/Engine/Source/Runtime/Experimental/RigidPhysics/Private/RigidPhysics/Internal/IJointConstraint.cpp
new file mode 100644
--- /dev/null
+++ b/Engine/Source/Runtime/Experimental/RigidPhysics/Private/RigidPhysics/Internal/IJointConstraint.cpp
@@ -0,0 +1,11 @@
+// Copyright Epic Games, Inc. All Rights Reserved.
+#include "RigidPhysics/Internal/IJointConstraint.h"
+
+#if UE_RIGIDPHYSICS_API_ENABLED
+
+namespace UE::Physics
+{
+	UE_RIGIDPHYSICS_RIGIDTYPED_IMPL(IJointConstraint, IRigidTyped);
+} // namespace UE::Physics
+
+#endif // UE_RIGIDPHYSICS_API_ENABLED
diff --git a/Engine/Source/Runtime/Experimental/RigidPhysics/Private/RigidPhysics/Internal/IJointConstraint6DOF.cpp b/Engine/Source/Runtime/Experimental/RigidPhysics/Private/RigidPhysics/Internal/IJointConstraint6DOF.cpp
new file mode 100644
--- /dev/null
+++ b/Engine/Source/Runtime/Experimental/RigidPhysics/Private/RigidPhysics/Internal/IJointConstraint6DOF.cpp
@@ -0,0 +1,18 @@
+// Copyright Epic Games, Inc. All Rights Reserved.
+#include "RigidPhysics/Internal/IJointConstraint6DOF.h"
+
+#if UE_RIGIDPHYSICS_API_ENABLED
+
+#include "RigidPhysics/RigidLog.h"
+
+namespace UE::Physics
+{
+	UE_RIGIDPHYSICS_RIGIDTYPED_IMPL(IJointConstraint6DOF, IJointConstraint);
+
+	FJointConstraint6DOFHandle IJointConstraint6DOF::GetTypedHandle() const
+	{
+		return FJointConstraint6DOFHandle(GetHandle());
+	}
+} // namespace UE::Physics
+
+#endif // UE_RIGIDPHYSICS_API_ENABLED
diff --git a/Engine/Source/Runtime/Experimental/RigidPhysics/Private/RigidPhysics/Internal/IRigidScene.cpp b/Engine/Source/Runtime/Experimental/RigidPhysics/Private/RigidPhysics/Internal/IRigidScene.cpp
--- a/Engine/Source/Runtime/Experimental/RigidPhysics/Private/RigidPhysics/Internal/IRigidScene.cpp
+++ b/Engine/Source/Runtime/Experimental/RigidPhysics/Private/RigidPhysics/Internal/IRigidScene.cpp
@@ -15,6 +15,7 @@ namespace UE::Physics
 	{
 		// TODO_CHAOSAPI: registry slack should be a parameter
 		BodyContainerRegistry = MakeUnique<FRigidBodyContainerRegistry>(10);
+		JointConstraintRegistry = MakeUnique<FJointConstraintRegistry>(10);
 	}
 
 	IRigidBodyContainer* FRigidScene::GetBodyContainer(const FRigidObjectId& InBodyContainerId) const
@@ -31,10 +32,21 @@ namespace UE::Physics
 		return nullptr;
 	}
 
+	void FRigidScene::DestructScene()
+	{
+		BodyContainerRegistry.Reset();
+		JointConstraintRegistry.Reset();
+	}
+
 	FRigidBodyContainerRegistry* FRigidScene::GetBodyContainerRegistry() const
 	{
 		return BodyContainerRegistry.Get();
 	}
+
+	FJointConstraintRegistry* FRigidScene::GetJointConstraintRegistry() const
+	{
+		return JointConstraintRegistry.Get();
+	}
 } // namespace UE::Physics
 
 #endif // UE_RIGIDPHYSICS_API_ENABLED
diff --git a/Engine/Source/Runtime/Experimental/RigidPhysics/Private/RigidPhysics/JointConstraint6DOFHandle.cpp b/Engine/Source/Runtime/Experimental/RigidPhysics/Private/RigidPhysics/JointConstraint6DOFHandle.cpp
new file mode 100644
--- /dev/null
+++ b/Engine/Source/Runtime/Experimental/RigidPhysics/Private/RigidPhysics/JointConstraint6DOFHandle.cpp
@@ -0,0 +1,21 @@
+// Copyright Epic Games, Inc. All Rights Reserved.
+#include "RigidPhysics/JointConstraint6DOFHandle.h"
+
+#if UE_RIGIDPHYSICS_API_ENABLED
+
+#include "RigidPhysics/RigidLog.h"
+
+namespace UE::Physics
+{
+	FJointConstraint6DOFHandle::FJointConstraint6DOFHandle(const FRigidSceneHandle& InSceneHandle, const FRigidObjectId& InConstraintId)
+		: FJointConstraintHandle(InSceneHandle, InConstraintId)
+	{
+	}
+
+	FJointConstraint6DOFHandle::FJointConstraint6DOFHandle(const FJointConstraintHandle& InHandle)
+		: FJointConstraintHandle(InHandle)
+	{
+	}
+} // namespace UE::Physics
+
+#endif // UE_RIGIDPHYSICS_API_ENABLED
diff --git a/Engine/Source/Runtime/Experimental/RigidPhysics/Private/RigidPhysics/JointConstraintHandle.cpp b/Engine/Source/Runtime/Experimental/RigidPhysics/Private/RigidPhysics/JointConstraintHandle.cpp
new file mode 100644
--- /dev/null
+++ b/Engine/Source/Runtime/Experimental/RigidPhysics/Private/RigidPhysics/JointConstraintHandle.cpp
@@ -0,0 +1,39 @@
+// Copyright Epic Games, Inc. All Rights Reserved.
+#include "RigidPhysics/JointConstraintHandle.h"
+
+#if UE_RIGIDPHYSICS_API_ENABLED
+
+#include "RigidPhysics/RigidLog.h"
+#include "RigidPhysics/RigidScene.h"
+
+namespace UE::Physics
+{
+	FJointConstraintHandle::FJointConstraintHandle(const FRigidSceneHandle& InSceneHandle, const FRigidObjectId& InConstraintId)
+		: SceneHandle(InSceneHandle)
+		, ConstraintId(InConstraintId)
+	{
+	}
+
+	FRigidSceneHandle FJointConstraintHandle::GetSceneHandle() const
+	{
+		return SceneHandle;
+	}
+
+	const FRigidObjectId& FJointConstraintHandle::GetId() const
+	{
+		return ConstraintId;
+	}
+
+	void FJointConstraintHandle::Reset()
+	{
+		SceneHandle.Reset();
+		ConstraintId.Reset();
+	}
+
+	IJointConstraint* FJointConstraintHandle::PinInternal(const IRigidScene* Scene) const
+	{
+		return Scene->GetJointConstraint(ConstraintId);
+	}
+} // namespace UE::Physics
+
+#endif // UE_RIGIDPHYSICS_API_ENABLED
diff --git a/Engine/Source/Runtime/Experimental/RigidPhysics/Public/RigidPhysics/Internal/IJointConstraint.h b/Engine/Source/Runtime/Experimental/RigidPhysics/Public/RigidPhysics/Internal/IJointConstraint.h
new file mode 100644
--- /dev/null
+++ b/Engine/Source/Runtime/Experimental/RigidPhysics/Public/RigidPhysics/Internal/IJointConstraint.h
@@ -0,0 +1,25 @@
+// Copyright Epic Games, Inc. All Rights Reserved.
+#pragma once
+
+#include "RigidPhysics/RigidFwd.h"
+
+#if UE_RIGIDPHYSICS_API_ENABLED
+
+#include "RigidPhysics/JointConstraintHandle.h"
+#include "RigidPhysics/RigidTyped.h"
+
+namespace UE::Physics
+{
+	class UE_INTERNAL IJointConstraint : public IRigidTyped
+	{
+	public:
+		UE_RIGIDPHYSICS_RIGIDTYPED_DECL(RIGIDPHYSICS_API, IJointConstraint);
+
+		IJointConstraint() = default;
+		virtual ~IJointConstraint() = default;
+		
+		UE_INTERNAL RIGIDPHYSICS_API virtual FJointConstraintHandle GetHandle() const = 0;
+	};
+} // namespace UE::Physics
+
+#endif // UE_RIGIDPHYSICS_API_ENABLED
diff --git a/Engine/Source/Runtime/Experimental/RigidPhysics/Public/RigidPhysics/Internal/IJointConstraint6DOF.h b/Engine/Source/Runtime/Experimental/RigidPhysics/Public/RigidPhysics/Internal/IJointConstraint6DOF.h
new file mode 100644
--- /dev/null
+++ b/Engine/Source/Runtime/Experimental/RigidPhysics/Public/RigidPhysics/Internal/IJointConstraint6DOF.h
@@ -0,0 +1,86 @@
+// Copyright Epic Games, Inc. All Rights Reserved.
+#pragma once
+
+#include "RigidPhysics/RigidFwd.h"
+
+#if UE_RIGIDPHYSICS_API_ENABLED
+
+#include "Math/Transform.h"
+#include "RigidPhysics/JointConstraint.h"
+#include "RigidPhysics/JointConstraint6DOFHandle.h"
+
+namespace UE::Physics
+{
+	class UE_INTERNAL IJointConstraint6DOF : public IJointConstraint
+	{
+	public:
+		UE_RIGIDPHYSICS_RIGIDTYPED_DECL(RIGIDPHYSICS_API, IJointConstraint6DOF);
+
+		IJointConstraint6DOF() = default;
+		virtual ~IJointConstraint6DOF() = default;
+
+		UE_INTERNAL RIGIDPHYSICS_API virtual FJointConstraintHandle GetHandle() const = 0;
+		UE_INTERNAL RIGIDPHYSICS_API FJointConstraint6DOFHandle GetTypedHandle() const;
+
+		UE_INTERNAL RIGIDPHYSICS_API virtual void Activate()
+		{
+			RIGIDPHYSICS_API_UNSUPPORTED();
+		}
+
+		UE_INTERNAL RIGIDPHYSICS_API virtual void SetBodies(IRigidBody* Body0, IRigidBody* Body1)
+		{
+			RIGIDPHYSICS_API_UNSUPPORTED();
+		}
+
+		UE_INTERNAL RIGIDPHYSICS_API virtual void SetJointTransforms(const FTransform3d& InTransform0, const FTransform3d& InTransform1)
+		{
+			RIGIDPHYSICS_API_UNSUPPORTED();
+		}
+
+		UE_INTERNAL RIGIDPHYSICS_API virtual double GetLinearLimit() const
+		{
+			RIGIDPHYSICS_API_UNSUPPORTED();
+			return 0;
+		}
+
+		UE_INTERNAL RIGIDPHYSICS_API virtual void SetLinearLimit(const double InLinearLimit)
+		{
+			RIGIDPHYSICS_API_UNSUPPORTED();
+		}
+
+		UE_INTERNAL RIGIDPHYSICS_API virtual EJointMotionType GetLinearMotionTypesX() const
+		{
+			RIGIDPHYSICS_API_UNSUPPORTED();
+			return EJointMotionType::Free;
+		}
+
+		UE_INTERNAL RIGIDPHYSICS_API virtual void SetLinearMotionTypesX(EJointMotionType InMotionType)
+		{
+			RIGIDPHYSICS_API_UNSUPPORTED();
+		}
+
+		UE_INTERNAL RIGIDPHYSICS_API virtual EJointMotionType GetLinearMotionTypesY() const
+		{
+			RIGIDPHYSICS_API_UNSUPPORTED();
+			return EJointMotionType::Free;
+		}
+
+		UE_INTERNAL RIGIDPHYSICS_API virtual void SetLinearMotionTypesY(EJointMotionType InMotionType)
+		{
+			RIGIDPHYSICS_API_UNSUPPORTED();
+		}
+
+		UE_INTERNAL RIGIDPHYSICS_API virtual EJointMotionType GetLinearMotionTypesZ() const
+		{
+			RIGIDPHYSICS_API_UNSUPPORTED();
+			return EJointMotionType::Free;
+		}
+
+		UE_INTERNAL RIGIDPHYSICS_API virtual void SetLinearMotionTypesZ(EJointMotionType InMotionType)
+		{
+			RIGIDPHYSICS_API_UNSUPPORTED();
+		}
+	};
+} // namespace UE::Physics
+
+#endif // UE_RIGIDPHYSICS_API_ENABLED
diff --git a/Engine/Source/Runtime/Experimental/RigidPhysics/Public/RigidPhysics/Internal/IRigidScene.h b/Engine/Source/Runtime/Experimental/RigidPhysics/Public/RigidPhysics/Internal/IRigidScene.h
--- a/Engine/Source/Runtime/Experimental/RigidPhysics/Public/RigidPhysics/Internal/IRigidScene.h
+++ b/Engine/Source/Runtime/Experimental/RigidPhysics/Public/RigidPhysics/Internal/IRigidScene.h
@@ -73,6 +73,23 @@ namespace UE::Physics
 		{
 			RIGIDPHYSICS_API_UNSUPPORTED();
 		}
+		
+		UE_INTERNAL RIGIDPHYSICS_API virtual IJointConstraint6DOF* CreateJointConstraint6DOF()
+		{
+			RIGIDPHYSICS_API_UNSUPPORTED();
+			return nullptr;
+		}
+
+		UE_INTERNAL RIGIDPHYSICS_API virtual void DestroyJointConstraint(IJointConstraint* InConstraint)
+		{
+			RIGIDPHYSICS_API_UNSUPPORTED();
+		}
+		
+		UE_INTERNAL RIGIDPHYSICS_API virtual IJointConstraint* GetJointConstraint(const FRigidObjectId& InConstraintId) const
+		{
+			RIGIDPHYSICS_API_UNSUPPORTED();
+			return nullptr;
+		}
 	};
 
 	//
@@ -89,10 +106,14 @@ namespace UE::Physics
 		UE_INTERNAL RIGIDPHYSICS_API virtual IRigidGeometryCollection* GetGeometryCollection(const FRigidObjectId& InGeometryCollectionId) const override final;
 
 	protected:
+		UE_INTERNAL RIGIDPHYSICS_API void DestructScene();
+
 		UE_INTERNAL RIGIDPHYSICS_API FRigidBodyContainerRegistry* GetBodyContainerRegistry() const;
+		UE_INTERNAL RIGIDPHYSICS_API FJointConstraintRegistry* GetJointConstraintRegistry() const;
 
 	private:
 		TUniquePtr<FRigidBodyContainerRegistry> BodyContainerRegistry;
+		TUniquePtr<FJointConstraintRegistry> JointConstraintRegistry;
 	};
 } // namespace UE::Physics
 
diff --git a/Engine/Source/Runtime/Experimental/RigidPhysics/Public/RigidPhysics/JointConstraint.h b/Engine/Source/Runtime/Experimental/RigidPhysics/Public/RigidPhysics/JointConstraint.h
new file mode 100644
--- /dev/null
+++ b/Engine/Source/Runtime/Experimental/RigidPhysics/Public/RigidPhysics/JointConstraint.h
@@ -0,0 +1,81 @@
+// Copyright Epic Games, Inc. All Rights Reserved.
+#pragma once
+
+#include "RigidPhysics/RigidFwd.h"
+
+#if UE_RIGIDPHYSICS_API_ENABLED
+
+#include "RigidPhysics/Internal/IJointConstraint.h"
+#include "RigidPhysics/JointConstraintHandle.h"
+#include "RigidPhysics/RigidLog.h"
+#include "RigidPhysics/RigidObjectPtr.h"
+#include "RigidPhysics/RigidTyped.h"
+
+UE_EXPERIMENTAL(5.8, "The new Chaos API is experimental")
+namespace UE::Physics
+{
+	template <typename ContextType>
+	class TJointConstraintPtrImpl
+	{
+	public:
+		using FContext = ContextType;
+		using FInterface = IJointConstraint;
+		using FHandle = FJointConstraintHandle;
+
+		UE_INTERNAL TJointConstraintPtrImpl() = default;
+		UE_INTERNAL TJointConstraintPtrImpl(IJointConstraint* InConstraint)
+			: ConstraintRaw(InConstraint)
+		{
+		}
+
+		TJointConstraintPtrImpl(const TJointConstraintPtrImpl&) = delete;
+
+		TJointConstraintPtrImpl(TJointConstraintPtrImpl&& InPtr)
+			: ConstraintRaw(InPtr.ConstraintRaw)
+		{
+			InPtr.ConstraintRaw = nullptr;
+		}
+
+		~TJointConstraintPtrImpl()
+		{
+			Reset();
+		}
+
+		friend bool operator==(const TJointConstraintPtrImpl&, const TJointConstraintPtrImpl&) = default;
+		friend bool operator!=(const TJointConstraintPtrImpl&, const TJointConstraintPtrImpl&) = default;
+
+		bool IsValid() const
+		{
+			return (ConstraintRaw != nullptr);
+		}
+
+		FJointConstraintHandle GetHandle() const
+		{
+			if (ConstraintRaw != nullptr)
+			{
+				return ConstraintRaw->GetHandle();
+			}
+			return FJointConstraintHandle();
+		}
+
+		UE_INTERNAL void Reset()
+		{
+			ConstraintRaw = nullptr;
+		}
+
+		UE_INTERNAL const FRigidTypeId& GetTypeId() const
+		{
+			return ConstraintRaw->GetTypeId();
+		}
+
+		UE_INTERNAL IJointConstraint* Get() const
+		{
+			return ConstraintRaw;
+		}
+
+	private:
+		IJointConstraint* ConstraintRaw = nullptr;
+	};
+} // namespace UE::Physics
+
+#endif // UE_RIGIDPHYSICS_API_ENABLED
diff --git a/Engine/Source/Runtime/Experimental/RigidPhysics/Public/RigidPhysics/JointConstraint6DOF.h b/Engine/Source/Runtime/Experimental/RigidPhysics/Public/RigidPhysics/JointConstraint6DOF.h
new file mode 100644
--- /dev/null
+++ b/Engine/Source/Runtime/Experimental/RigidPhysics/Public/RigidPhysics/JointConstraint6DOF.h
@@ -0,0 +1,139 @@
+// Copyright Epic Games, Inc. All Rights Reserved.
+#pragma once
+
+#include "RigidPhysics/RigidFwd.h"
+
+#if UE_RIGIDPHYSICS_API_ENABLED
+
+#include "RigidPhysics/Internal/IJointConstraint6DOF.h"
+#include "RigidPhysics/JointConstraint.h"
+#include "RigidPhysics/JointConstraint6DOFHandle.h"
+
+UE_EXPERIMENTAL(5.8, "The new Chaos API is experimental")
+namespace UE::Physics
+{
+	template <typename ContextType>
+	class TJointConstraint6DOFPtrImpl
+	{
+	public:
+		using FContext = ContextType;
+		using FInterface = IJointConstraint6DOF;
+		using FHandle = FJointConstraint6DOFHandle;
+
+		UE_INTERNAL TJointConstraint6DOFPtrImpl() = default;
+		UE_INTERNAL TJointConstraint6DOFPtrImpl(IJointConstraint6DOF* InConstraint)
+			: ConstraintRaw(InConstraint)
+		{
+		}
+
+		// Construct from base class (does a safe downcast)
+		UE_INTERNAL TJointConstraint6DOFPtrImpl(IJointConstraint* InConstraint)
+			: ConstraintRaw(InConstraint ? InConstraint->AsA<IJointConstraint6DOF>() : nullptr)
+		{
+		}
+
+		TJointConstraint6DOFPtrImpl(const TJointConstraint6DOFPtrImpl&) = delete;
+
+		TJointConstraint6DOFPtrImpl(TJointConstraint6DOFPtrImpl&& InPtr)
+			: ConstraintRaw(InPtr.ConstraintRaw)
+		{
+			InPtr.ConstraintRaw = nullptr;
+		}
+
+		~TJointConstraint6DOFPtrImpl()
+		{
+			Reset();
+		}
+
+		friend bool operator==(const TJointConstraint6DOFPtrImpl&, const TJointConstraint6DOFPtrImpl&) = default;
+		friend bool operator!=(const TJointConstraint6DOFPtrImpl&, const TJointConstraint6DOFPtrImpl&) = default;
+
+		bool IsValid() const
+		{
+			return (ConstraintRaw != nullptr);
+		}
+
+		FJointConstraint6DOFHandle GetHandle() const
+		{
+			if (ConstraintRaw != nullptr)
+			{
+				return ConstraintRaw->GetTypedHandle();
+			}
+			return FJointConstraint6DOFHandle();
+		}
+
+		void Activate()
+		{
+			ConstraintRaw->Activate();
+		}
+
+		void SetBodies(const TRigidBodyPtr<ContextType>& Body0, const TRigidBodyPtr<ContextType>& Body1)
+		{
+			ConstraintRaw->SetBodies(Body0.Get(), Body1.Get());
+		}
+
+		void SetJointTransforms(const FTransform3d& InTransform0, const FTransform3d& InTransform1)
+		{
+			ConstraintRaw->SetJointTransforms(InTransform0, InTransform1);
+		}
+
+		double GetLinearLimit() const
+		{
+			return ConstraintRaw->GetLinearLimit();
+		}
+
+		void SetLinearLimit(const double InLinearLimit)
+		{
+			ConstraintRaw->SetLinearLimit(InLinearLimit);
+		}
+
+		EJointMotionType GetLinearMotionTypesX()
+		{
+			return ConstraintRaw->GetLinearMotionTypesX();
+		}
+
+		void SetLinearMotionTypesX(const EJointMotionType InMotionType)
+		{
+			ConstraintRaw->SetLinearMotionTypesX(InMotionType);
+		}
+
+		EJointMotionType GetLinearMotionTypesY()
+		{
+			return ConstraintRaw->GetLinearMotionTypesY();
+		}
+
+		void SetLinearMotionTypesY(const EJointMotionType InMotionType)
+		{
+			ConstraintRaw->SetLinearMotionTypesY(InMotionType);
+		}
+
+		EJointMotionType GetLinearMotionTypesZ()
+		{
+			return ConstraintRaw->GetLinearMotionTypesZ();
+		}
+
+		void SetLinearMotionTypesZ(const EJointMotionType InMotionType)
+		{
+			ConstraintRaw->SetLinearMotionTypesZ(InMotionType);
+		}
+
+		UE_INTERNAL void Reset()
+		{
+			ConstraintRaw = nullptr;
+		}
+
+		UE_INTERNAL const FRigidTypeId& GetTypeId() const
+		{
+			return ConstraintRaw->GetTypeId();
+		}
+
+		UE_INTERNAL IJointConstraint6DOF* Get() const
+		{
+			return ConstraintRaw;
+		}
+	private:
+		IJointConstraint6DOF* ConstraintRaw = nullptr;
+	};
+} // namespace UE::Physics
+
+#endif // UE_RIGIDPHYSICS_API_ENABLED
diff --git a/Engine/Source/Runtime/Experimental/RigidPhysics/Public/RigidPhysics/JointConstraint6DOFHandle.h b/Engine/Source/Runtime/Experimental/RigidPhysics/Public/RigidPhysics/JointConstraint6DOFHandle.h
new file mode 100644
--- /dev/null
+++ b/Engine/Source/Runtime/Experimental/RigidPhysics/Public/RigidPhysics/JointConstraint6DOFHandle.h
@@ -0,0 +1,41 @@
+// Copyright Epic Games, Inc. All Rights Reserved.
+#pragma once
+
+#include "RigidPhysics/RigidFwd.h"
+
+#if UE_RIGIDPHYSICS_API_ENABLED
+
+#include "RigidPhysics/JointConstraintHandle.h"
+
+UE_EXPERIMENTAL(5.8, "The new Chaos API is experimental")
+namespace UE::Physics
+{
+	class IJointConstraint6DOF;
+	template <typename ContextType>
+	class UE_INTERNAL TJointConstraint6DOFPtrImpl;
+
+	class FJointConstraint6DOFHandle : public FJointConstraintHandle
+	{
+	public:
+		FJointConstraint6DOFHandle() = default;
+		FJointConstraint6DOFHandle(const FJointConstraint6DOFHandle&) = default;
+		UE_INTERNAL RIGIDPHYSICS_API FJointConstraint6DOFHandle(const FRigidSceneHandle& InSceneHandle, const FRigidObjectId& InConstraintId);
+
+		friend bool operator==(const FJointConstraint6DOFHandle&, const FJointConstraint6DOFHandle&) = default;
+		friend bool operator!=(const FJointConstraint6DOFHandle&, const FJointConstraint6DOFHandle&) = default;
+
+		template <typename ContextType>
+		TJointConstraint6DOFPtr<ContextType> Pin(const ContextType& InContext) const
+		{
+			return FJointConstraintHandle::Pin(InContext);
+		}
+
+	private:
+		friend class IJointConstraint6DOF;
+
+		// Private so we only allow downcast in the context where we explicitly know this constraint handle is the correct type.
+		UE_INTERNAL RIGIDPHYSICS_API FJointConstraint6DOFHandle(const FJointConstraintHandle& InHandle);
+	};
+} // namespace UE::Physics
+
+#endif // UE_RIGIDPHYSICS_API_ENABLED
diff --git a/Engine/Source/Runtime/Experimental/RigidPhysics/Public/RigidPhysics/JointConstraintHandle.h b/Engine/Source/Runtime/Experimental/RigidPhysics/Public/RigidPhysics/JointConstraintHandle.h
new file mode 100644
--- /dev/null
+++ b/Engine/Source/Runtime/Experimental/RigidPhysics/Public/RigidPhysics/JointConstraintHandle.h
@@ -0,0 +1,53 @@
+// Copyright Epic Games, Inc. All Rights Reserved.
+#pragma once
+
+#include "RigidPhysics/RigidFwd.h"
+
+#if UE_RIGIDPHYSICS_API_ENABLED
+
+#include "RigidPhysics/RigidObjectId.h"
+#include "RigidPhysics/RigidObjectPtr.h"
+#include "RigidPhysics/RigidSceneHandle.h"
+
+UE_EXPERIMENTAL(5.8, "The new Chaos API is experimental")
+namespace UE::Physics
+{
+	class IJointConstraint;
+	template <typename ContextType>
+	class UE_INTERNAL TJointConstraintPtrImpl;
+
+	class FJointConstraintHandle
+	{
+	public:
+		FJointConstraintHandle() = default;
+		FJointConstraintHandle(const FJointConstraintHandle&) = default;
+		UE_INTERNAL RIGIDPHYSICS_API FJointConstraintHandle(const FRigidSceneHandle& InSceneHandle, const FRigidObjectId& InConstraintId);
+
+		friend bool operator==(const FJointConstraintHandle&, const FJointConstraintHandle&) = default;
+		friend bool operator!=(const FJointConstraintHandle&, const FJointConstraintHandle&) = default;
+
+		RIGIDPHYSICS_API FRigidSceneHandle GetSceneHandle() const;
+
+		template <typename ContextType>
+		TJointConstraintPtr<ContextType> Pin(const ContextType& InContext) const
+		{
+			if (InContext.IsValid())
+			{
+				const IRigidScene* Scene = InContext.GetScene().Get();
+				return TJointConstraintPtr<ContextType>(PinInternal(Scene));
+			}
+			return TJointConstraintPtr<ContextType>();
+		}
+
+		UE_INTERNAL RIGIDPHYSICS_API const FRigidObjectId& GetId() const;
+		UE_INTERNAL RIGIDPHYSICS_API void Reset();
+
+	private:
+		RIGIDPHYSICS_API IJointConstraint* PinInternal(const IRigidScene* Scene) const;
+
+		FRigidSceneHandle SceneHandle;
+		FRigidObjectId ConstraintId;
+	};
+} // namespace UE::Physics
+
+#endif // UE_RIGIDPHYSICS_API_ENABLED
diff --git a/Engine/Source/Runtime/Experimental/RigidPhysics/Public/RigidPhysics/RigidFwd.h b/Engine/Source/Runtime/Experimental/RigidPhysics/Public/RigidPhysics/RigidFwd.h
--- a/Engine/Source/Runtime/Experimental/RigidPhysics/Public/RigidPhysics/RigidFwd.h
+++ b/Engine/Source/Runtime/Experimental/RigidPhysics/Public/RigidPhysics/RigidFwd.h
@@ -55,6 +55,8 @@ struct FSimulationParameters;	// GC setup
 // User-facing classes
 namespace UE::Physics
 {
+	class FJointConstraintHandle;
+	class FJointConstraint6DOFHandle;
 	class FRigidBodyContainerHandle;
 	class FRigidBodyHandle;
 	class FRigidGeometryCollectionHandle;
@@ -64,6 +66,8 @@ namespace UE::Physics
 	class FRigidSceneHandle;
 	class FRigidShapeInstance;
 
+	class IJointConstraint;
+	class IJointConstraint6DOF;
 	class IRigidBody;
 	class IRigidBodyContainer;
 	class IRigidFactory;
@@ -75,6 +79,8 @@ namespace UE::Physics
 	class IRigidSceneSettings;
 	class IRigidTyped;
 
+	template<typename ContextType> class TJointConstraintPtrImpl;
+	template<typename ContextType> class TJointConstraint6DOFPtrImpl;
 	template<typename ContextType> class TRigidBodyContainerPtrImpl;
 	template<typename ContextType> class TRigidBodyPtrImpl;
 	template<typename ContextType> class TRigidGeometryCollectionPtrImpl;
@@ -82,6 +88,8 @@ namespace UE::Physics
 	template<typename ContextType> class TRigidScenePtrImpl;
 	template<bool bConstParam> class TRigidShapeInstancePtr;
 
+	template<typename ContextType> using TJointConstraintPtr = TRigidObjectPtr<TJointConstraintPtrImpl<ContextType>>;
+	template<typename ContextType> using TJointConstraint6DOFPtr = TRigidObjectPtr<TJointConstraint6DOFPtrImpl<ContextType>>;
 	template<typename ContextType> using TRigidBodyContainerPtr = TRigidObjectPtr<TRigidBodyContainerPtrImpl<ContextType>>;
 	template<typename ContextType> using TRigidBodyPtr = TRigidObjectPtr<TRigidBodyPtrImpl<ContextType>>;
 	template<typename ContextType> using TRigidGeometryCollectionPtr = TRigidObjectPtr<TRigidGeometryCollectionPtrImpl<ContextType>>;
@@ -95,6 +103,13 @@ namespace UE::Physics
 
 namespace UE::Physics
 {
+	enum class EJointMotionType : int32
+	{
+		Free,
+		Limited,
+		Locked,
+	};
+
 	enum class ERigidVisitorResponse
 	{
 		Continue,
@@ -138,6 +153,7 @@ namespace UE::Physics
 namespace UE::Physics
 {
 	template<typename ObjectType> class UE_INTERNAL TRigidObjectRegistry;
+	using UE_INTERNAL FJointConstraintRegistry = TRigidObjectRegistry<IJointConstraint>;
 	using UE_INTERNAL FRigidBodyContainerRegistry = TRigidObjectRegistry<IRigidBodyContainer>;
 	using UE_INTERNAL FRigidGeometryCollectionRegistry = TRigidObjectRegistry<IRigidGeometryCollection>;
 	using UE_INTERNAL FRigidBodyRegistry = TRigidObjectRegistry<IRigidBody>;
diff --git a/Engine/Source/Runtime/Experimental/RigidPhysics/Public/RigidPhysics/RigidScene.h b/Engine/Source/Runtime/Experimental/RigidPhysics/Public/RigidPhysics/RigidScene.h
--- a/Engine/Source/Runtime/Experimental/RigidPhysics/Public/RigidPhysics/RigidScene.h
+++ b/Engine/Source/Runtime/Experimental/RigidPhysics/Public/RigidPhysics/RigidScene.h
@@ -99,6 +99,25 @@ namespace UE::Physics
 			SceneRaw->DestroyGeometryCollection(InGC.Get());
 			InGC.Reset();
 		}
+		
+		TJointConstraint6DOFPtr<FContext> CreateJointConstraint6DOF()
+		{
+			return SceneRaw->CreateJointConstraint6DOF();
+		}
+
+		void DestroyJointConstraint(TJointConstraint6DOFPtr<FContext>&& InConstraint)
+		{
+			if (InConstraint.IsValid())
+			{
+				SceneRaw->DestroyJointConstraint(InConstraint.Get());
+				InConstraint.Reset();
+			}
+		}
+
+		UE_INTERNAL IJointConstraint* GetJointConstraint(const FRigidObjectId& InConstraintId) const
+		{
+			return SceneRaw->GetJointConstraint(InConstraintId);
+		}
 
 		bool Query(const FRigidQueryDesc& QueryDesc, TArray<FRigidQueryHit>& OutHits)  const
 		{
