Chaos API: Adding new core geometry types and the AnyGeometry type plus tests.
The convex and trimesh geometry objects are more stubbed for now, more work will have to be done on these later.

#Synced-CL 51413247

diff --git a/Engine/Plugins/Experimental/ChaosRigidPhysicsAsync/Tests/Private/AnyGeometryTests.cpp b/Engine/Plugins/Experimental/ChaosRigidPhysicsAsync/Tests/Private/AnyGeometryTests.cpp
new file mode 100644
--- /dev/null
+++ b/Engine/Plugins/Experimental/ChaosRigidPhysicsAsync/Tests/Private/AnyGeometryTests.cpp
@@ -0,0 +1,129 @@
+// Copyright Epic Games, Inc. All Rights Reserved.
+
+#include "Chaos/LowLevelTest/ChaosTestHarness.h"
+#include "Chaos/LowLevelTest/ChaosTestScene.h"
+#include "ChaosRigidPhysicsAsyncTest.h"
+
+#include "RigidPhysics/Geometry/AnyGeometry.h"
+
+#if UE_RIGIDPHYSICS_API_ENABLED
+
+namespace Chaos::LowLevelTest
+{
+	using FSphereGeometry = UE::Physics::FSphereGeometry;
+	using FBoxGeometry = UE::Physics::FBoxGeometry;
+	using FAnyGeometry = UE::Physics::FAnyGeometry;
+
+	template <typename ExpectedGeometryType, typename ActualGeometryType>
+	void CheckGeometry(const ExpectedGeometryType& ExpectedGeometry, const ActualGeometryType& ActualGeometry)
+	{
+		constexpr bool bIsExpectedType = std::is_same_v<ExpectedGeometryType, ActualGeometryType>;
+		CHECK(bIsExpectedType);
+		if constexpr (bIsExpectedType)
+		{
+			CHECK(ExpectedGeometry == ActualGeometry);
+		}
+	}
+
+	TEST_CASE("AnyGeometry: None", "[Chaos][API][Geometry][unit]")
+	{
+		// Defaults to the first shape type, which is a default constructed sphere.
+		FAnyGeometry AnyGeometry;
+		REQUIRE(AnyGeometry.IsA<FSphereGeometry>());
+		CHECK(AnyGeometry.Get<FSphereGeometry>() == FSphereGeometry());
+	}
+
+	TEST_CASE("AnyGeometry: Box", "[Chaos][API][Geometry][unit]")
+	{
+		FBoxGeometry Box0(FVector3f(1, 2, 3));
+		FBoxGeometry Box1(FVector3f(2, 3, 4));
+		FSphereGeometry Sphere0(1);
+
+		SECTION("IsA")
+		{
+			FAnyGeometry AnyGeometry(Box0);
+			CHECK(AnyGeometry.IsA<FBoxGeometry>());
+			CHECK(!AnyGeometry.IsA<FSphereGeometry>());
+		}
+		SECTION("TryGet/Get")
+		{
+			FAnyGeometry AnyGeometry(Box0);
+			REQUIRE(AnyGeometry.TryGet<FBoxGeometry>() != nullptr);
+			CHECK(*AnyGeometry.TryGet<FBoxGeometry>() == Box0);
+			CHECK(AnyGeometry.Get<FBoxGeometry>() == Box0);
+			CHECK(AnyGeometry.TryGet<FSphereGeometry>() == nullptr);
+		}
+		SECTION("Visit")
+		{
+			FAnyGeometry AnyGeometry(Box0);
+			AnyGeometry.Visit([&Box0]<typename GeometryType>(const GeometryType& Geometry) { CheckGeometry(Box0, Geometry); });
+		}
+		SECTION("Comparison")
+		{
+			FAnyGeometry AnyGeometry(Box0);
+			CHECK(AnyGeometry == Box0);
+			CHECK(AnyGeometry != Box1);
+			CHECK(AnyGeometry != Sphere0);
+		}
+		SECTION("Assignment")
+		{
+			FAnyGeometry AnyGeometry(Sphere0);
+			AnyGeometry = Box0;
+			CHECK(AnyGeometry != Sphere0);
+			CHECK(AnyGeometry == Box0);
+			CHECK(AnyGeometry != Box1);
+			AnyGeometry = Box1;
+			CHECK(AnyGeometry != Sphere0);
+			CHECK(AnyGeometry != Box0);
+			CHECK(AnyGeometry == Box1);
+		}
+	}
+
+	TEST_CASE("AnyGeometry: Sphere", "[Chaos][API][Geometry][unit]")
+	{
+		FBoxGeometry Box0(FVector3f(1, 2, 3));
+		FSphereGeometry Sphere0(1);
+		FSphereGeometry Sphere1(2);
+
+		SECTION("IsA")
+		{
+			FAnyGeometry AnyGeometry(Sphere0);
+			CHECK(AnyGeometry.IsA<FSphereGeometry>());
+			CHECK(!AnyGeometry.IsA<FBoxGeometry>());
+		}
+		SECTION("TryGet/Get")
+		{
+			FAnyGeometry AnyGeometry(Sphere0);
+			REQUIRE(AnyGeometry.TryGet<FSphereGeometry>() != nullptr);
+			CHECK(*AnyGeometry.TryGet<FSphereGeometry>() == Sphere0);
+			CHECK(AnyGeometry.Get<FSphereGeometry>() == Sphere0);
+			CHECK(AnyGeometry.TryGet<FBoxGeometry>() == nullptr);
+		}
+		SECTION("Visit")
+		{
+			FAnyGeometry AnyGeometry(Sphere0);
+			AnyGeometry.Visit([&Sphere0]<typename GeometryType>(const GeometryType& Geometry) { CheckGeometry(Sphere0, Geometry); });
+		}
+		SECTION("Comparison")
+		{
+			FAnyGeometry AnyGeometry(Sphere0);
+			CHECK(AnyGeometry == Sphere0);
+			CHECK(AnyGeometry != Sphere1);
+			CHECK(AnyGeometry != Box0);
+		}
+		SECTION("Assignment")
+		{
+			FAnyGeometry AnyGeometry(Box0);
+			AnyGeometry = Sphere0;
+			CHECK(AnyGeometry != Box0);
+			CHECK(AnyGeometry == Sphere0);
+			CHECK(AnyGeometry != Sphere1);
+			AnyGeometry = Sphere1;
+			CHECK(AnyGeometry != Box0);
+			CHECK(AnyGeometry != Sphere0);
+			CHECK(AnyGeometry == Sphere1);
+		}
+	}
+} // namespace Chaos::LowLevelTest
+
+#endif // UE_RIGIDPHYSICS_API_ENABLED
diff --git a/Engine/Plugins/Experimental/ChaosRigidPhysicsAsync/Tests/Private/ChaosRigidPhysicsAsyncTest.h b/Engine/Plugins/Experimental/ChaosRigidPhysicsAsync/Tests/Private/ChaosRigidPhysicsAsyncTest.h
--- a/Engine/Plugins/Experimental/ChaosRigidPhysicsAsync/Tests/Private/ChaosRigidPhysicsAsyncTest.h
+++ b/Engine/Plugins/Experimental/ChaosRigidPhysicsAsync/Tests/Private/ChaosRigidPhysicsAsyncTest.h
@@ -3,6 +3,7 @@
 
 #include "Chaos/LowLevelTest/ChaosTestHarness.h"
 #include "Chaos/LowLevelTest/ChaosTestScene.h"
+#include "RigidPhysics/Geometry/AnyGeometry.h"
 #include "RigidPhysics/RigidTyped.h"
 
 #include <string>
@@ -23,5 +24,20 @@ namespace Catch
 	{
 		static std::string convert(UE::Physics::FRigidTypeId const& Value) { return Chaos::LowLevelTest::ToStdString(Value.GetTypeName()); }
 	};
+	template<>
+	struct StringMaker<UE::Physics::FSphereGeometry>
+	{
+		static std::string convert(UE::Physics::FSphereGeometry const& Value) { return StringMaker<FString>::convert(Value.ToString()); }
+	};
+	template<>
+	struct StringMaker<UE::Physics::FBoxGeometry>
+	{
+		static std::string convert(UE::Physics::FBoxGeometry const& Value) { return StringMaker<FString>::convert(Value.ToString()); }
+	};
+	template<>
+	struct StringMaker<UE::Physics::FAnyGeometry>
+	{
+		static std::string convert(UE::Physics::FAnyGeometry const& Value) { return StringMaker<FString>::convert(Value.ToString()); }
+	};
 }
 
diff --git a/Engine/Plugins/Experimental/ChaosRigidPhysicsAsync/Tests/Private/JointConstraintTests.cpp b/Engine/Plugins/Experimental/ChaosRigidPhysicsAsync/Tests/Private/JointConstraintTests.cpp
--- a/Engine/Plugins/Experimental/ChaosRigidPhysicsAsync/Tests/Private/JointConstraintTests.cpp
+++ b/Engine/Plugins/Experimental/ChaosRigidPhysicsAsync/Tests/Private/JointConstraintTests.cpp
@@ -20,6 +20,7 @@ namespace Chaos::LowLevelTest
 	using namespace Chaos;
 	using namespace Chaos::Rigids::Async;
 	using namespace UE::Physics;
+	using EJointMotionType = UE::Physics::EJointMotionType;
 	
 	struct FJointFixture : public FRigidTestFixture
 	{
diff --git a/Engine/Source/Runtime/Experimental/RigidPhysics/Private/RigidPhysics/Geometry/AnyGeometry.cpp b/Engine/Source/Runtime/Experimental/RigidPhysics/Private/RigidPhysics/Geometry/AnyGeometry.cpp
new file mode 100644
--- /dev/null
+++ b/Engine/Source/Runtime/Experimental/RigidPhysics/Private/RigidPhysics/Geometry/AnyGeometry.cpp
@@ -0,0 +1,42 @@
+// Copyright Epic Games, Inc. All Rights Reserved.
+#include "RigidPhysics/Geometry/AnyGeometry.h"
+
+#if UE_RIGIDPHYSICS_API_ENABLED
+
+namespace UE::Physics
+{
+	template <typename GeometryType>
+	FString GeometryToString(const GeometryType& Geometry)
+	{
+		return Geometry.ToString();
+	}
+
+	template <typename GeometryType>
+	FString GeometryToString(const TRefCountPtr<GeometryType>& Geometry)
+	{
+		return Geometry.IsValid() ? Geometry->ToString() : FString("Null");
+	}
+
+	bool FAnyGeometry::operator==(const FAnyGeometry& Rhs) const
+	{
+		if (Geometry.GetIndex() != Rhs.Geometry.GetIndex())
+		{
+			return false;
+		}
+		return Visit([&Rhs]<typename GeometryType>(const GeometryType& SelfGeometry)
+		{
+			const GeometryType& RhsGeometry = Rhs.Geometry.template Get<GeometryType>();
+			return SelfGeometry == RhsGeometry;
+		});
+	}
+
+	FString FAnyGeometry::ToString() const
+	{
+		return Visit([]<typename GeometryType>(const GeometryType& Geometry)
+		{
+			return GeometryToString(Geometry);
+		});
+	}
+} // namespace UE::Physics
+
+#endif // UE_RIGIDPHYSICS_API_ENABLED
diff --git a/Engine/Source/Runtime/Experimental/RigidPhysics/Private/RigidPhysics/Geometry/CommonCoreGeometry.cpp b/Engine/Source/Runtime/Experimental/RigidPhysics/Private/RigidPhysics/Geometry/CommonCoreGeometry.cpp
new file mode 100644
--- /dev/null
+++ b/Engine/Source/Runtime/Experimental/RigidPhysics/Private/RigidPhysics/Geometry/CommonCoreGeometry.cpp
@@ -0,0 +1,80 @@
+// Copyright Epic Games, Inc. All Rights Reserved.
+#include "RigidPhysics/Geometry/CommonCoreGeometry.h"
+
+#if UE_RIGIDPHYSICS_API_ENABLED
+
+namespace UE::Physics
+{
+	FSphereGeometry::FSphereGeometry(const float InRadius)
+		: Radius(InRadius)
+	{
+	}
+
+	float FSphereGeometry::GetRadius() const
+	{
+		return Radius;
+	}
+
+	void FSphereGeometry::SetRadius(float InRadius)
+	{
+		Radius = InRadius;
+	}
+
+	FString FSphereGeometry::ToString() const
+	{
+		return FString::Printf(TEXT("Sphere: Radius: %f"), GetRadius());
+	}
+
+	FBoxGeometry::FBoxGeometry(const FVector3f& InExtents)
+		: Extents(InExtents)
+	{
+	}
+
+	FVector3f FBoxGeometry::GetExtents() const
+	{
+		return Extents;
+	}
+
+	void FBoxGeometry::SetExtents(const FVector3f& InExtents)
+	{
+		Extents = InExtents;
+	}
+
+	FString FBoxGeometry::ToString() const
+	{
+		return FString::Printf(TEXT("Box: Extents: [%s]"), *Extents.ToString());
+	}
+
+	FCapsuleGeometry::FCapsuleGeometry(float InRadius, float InSegmentHalfLength)
+		: Radius(InRadius)
+		, SegmentHalfLength(InSegmentHalfLength)
+	{
+	}
+
+	float FCapsuleGeometry::GetRadius() const
+	{
+		return Radius;
+	}
+
+	void FCapsuleGeometry::SetRadius(const float InRadius)
+	{
+		Radius = InRadius;
+	}
+
+	float FCapsuleGeometry::GetSegmentHalfLength() const
+	{
+		return SegmentHalfLength;
+	}
+
+	void FCapsuleGeometry::SetSegmentHalfLength(const float InHalfLength)
+	{
+		SegmentHalfLength = InHalfLength;
+	}
+
+	FString FCapsuleGeometry::ToString() const
+	{
+		return FString::Printf(TEXT("Capsule: Radius: %f, SegmentHalfLength: %f"), Radius, SegmentHalfLength);
+	}
+} // namespace UE::Physics
+
+#endif // UE_RIGIDPHYSICS_API_ENABLED
diff --git a/Engine/Source/Runtime/Experimental/RigidPhysics/Private/RigidPhysics/Geometry/ConvexGeometry.cpp b/Engine/Source/Runtime/Experimental/RigidPhysics/Private/RigidPhysics/Geometry/ConvexGeometry.cpp
new file mode 100644
--- /dev/null
+++ b/Engine/Source/Runtime/Experimental/RigidPhysics/Private/RigidPhysics/Geometry/ConvexGeometry.cpp
@@ -0,0 +1,19 @@
+// Copyright Epic Games, Inc. All Rights Reserved.
+#include "RigidPhysics/Geometry/ConvexGeometry.h"
+
+#if UE_RIGIDPHYSICS_API_ENABLED
+
+namespace UE::Physics
+{
+	FConvexGeometry::FConvexGeometry(Chaos::FImplicitConvex3* InImplicit)
+		: Implicit(InImplicit)
+	{
+	}
+
+	FString FConvexGeometry::ToString() const
+	{
+		return TEXT("Convex");
+	}
+} // namespace UE::Physics
+
+#endif // UE_RIGIDPHYSICS_API_ENABLED
diff --git a/Engine/Source/Runtime/Experimental/RigidPhysics/Private/RigidPhysics/Geometry/TriangleMeshGeometry.cpp b/Engine/Source/Runtime/Experimental/RigidPhysics/Private/RigidPhysics/Geometry/TriangleMeshGeometry.cpp
new file mode 100644
--- /dev/null
+++ b/Engine/Source/Runtime/Experimental/RigidPhysics/Private/RigidPhysics/Geometry/TriangleMeshGeometry.cpp
@@ -0,0 +1,19 @@
+// Copyright Epic Games, Inc. All Rights Reserved.
+#include "RigidPhysics/Geometry/TriangleMeshGeometry.h"
+
+#if UE_RIGIDPHYSICS_API_ENABLED
+
+namespace UE::Physics
+{
+	FTriangleMeshGeometry::FTriangleMeshGeometry(Chaos::FTriangleMeshImplicitObject* InImplicit)
+		: Implicit(InImplicit)
+	{
+	}
+
+	FString FTriangleMeshGeometry::ToString() const
+	{
+		return TEXT("TriangleMesh");
+	}
+} // namespace UE::Physics
+
+#endif // UE_RIGIDPHYSICS_API_ENABLED
diff --git a/Engine/Source/Runtime/Experimental/RigidPhysics/Public/RigidPhysics/Geometry/AnyGeometry.h b/Engine/Source/Runtime/Experimental/RigidPhysics/Public/RigidPhysics/Geometry/AnyGeometry.h
new file mode 100644
--- /dev/null
+++ b/Engine/Source/Runtime/Experimental/RigidPhysics/Public/RigidPhysics/Geometry/AnyGeometry.h
@@ -0,0 +1,89 @@
+// Copyright Epic Games, Inc. All Rights Reserved.
+#pragma once
+
+#include "RigidPhysics/RigidFwd.h"
+
+#if UE_RIGIDPHYSICS_API_ENABLED
+
+#include "Misc/TVariant.h"
+#include "RigidPhysics/Geometry/CommonCoreGeometry.h"
+#include "RigidPhysics/Geometry/ConvexGeometry.h"
+#include "RigidPhysics/Geometry/TriangleMeshGeometry.h"
+
+UE_EXPERIMENTAL(5.8, "The new Chaos API is experimental")
+namespace UE::Physics
+{
+	class FAnyGeometry
+	{
+	public:
+		FAnyGeometry() = default;
+		FAnyGeometry(const FAnyGeometry&) = default;
+		FAnyGeometry(FAnyGeometry&&) = default;
+
+		template <typename GeometryType>
+		FAnyGeometry(const GeometryType& InGeometry)
+		{
+			Geometry.Set<GeometryType>(InGeometry);
+		}
+
+		FAnyGeometry& operator=(const FAnyGeometry& InGeometry) = default;
+
+		template <typename GeometryType>
+		FAnyGeometry& operator=(const GeometryType& InGeometry)
+		{
+			Geometry.Set<GeometryType>(InGeometry);
+			return *this;
+		}
+
+		bool operator==(const FAnyGeometry& Rhs) const;
+		bool operator!=(const FAnyGeometry& Rhs) const = default;
+
+		template <typename GeometryType>
+		bool IsA() const
+		{
+			return Geometry.IsType<GeometryType>();
+		}
+
+		template <typename GeometryType>
+		GeometryType& Get()
+		{
+			return Geometry.Get<GeometryType>();
+		}
+
+		template <typename GeometryType>
+		const GeometryType& Get() const
+		{
+			return Geometry.Get<GeometryType>();
+		}
+
+		template <typename GeometryType>
+		GeometryType* TryGet()
+		{
+			return Geometry.TryGet<GeometryType>();
+		}
+
+		template <typename GeometryType>
+		const GeometryType* TryGet() const
+		{
+			return Geometry.TryGet<GeometryType>();
+		}
+
+		template <typename TLambda>
+		decltype(auto) Visit(const TLambda& Lambda)
+		{
+			return ::Visit(Lambda, Geometry);
+		}
+
+		template <typename TLambda>
+		decltype(auto) Visit(const TLambda& Lambda) const
+		{
+			return ::Visit(Lambda, Geometry);
+		}
+
+		RIGIDPHYSICS_API FString ToString() const;
+	private:
+		TVariant<FSphereGeometry, FBoxGeometry, FCapsuleGeometry, FConvexGeometry, FTriangleMeshGeometry> Geometry;
+	};
+} // namespace UE::Physics
+
+#endif // UE_RIGIDPHYSICS_API_ENABLED
diff --git a/Engine/Source/Runtime/Experimental/RigidPhysics/Public/RigidPhysics/Geometry/CommonCoreGeometry.h b/Engine/Source/Runtime/Experimental/RigidPhysics/Public/RigidPhysics/Geometry/CommonCoreGeometry.h
new file mode 100644
--- /dev/null
+++ b/Engine/Source/Runtime/Experimental/RigidPhysics/Public/RigidPhysics/Geometry/CommonCoreGeometry.h
@@ -0,0 +1,69 @@
+// Copyright Epic Games, Inc. All Rights Reserved.
+#pragma once
+
+#include "RigidPhysics/RigidFwd.h"
+
+#if UE_RIGIDPHYSICS_API_ENABLED
+
+#include "CoreMinimal.h"
+
+UE_EXPERIMENTAL(5.8, "The new Chaos API is experimental")
+namespace UE::Physics
+{
+	struct FSphereGeometry
+	{
+		FSphereGeometry() = default;
+		RIGIDPHYSICS_API FSphereGeometry(const float InRadius);
+
+		bool operator==(const FSphereGeometry&) const = default;
+		bool operator!=(const FSphereGeometry&) const = default;
+
+		RIGIDPHYSICS_API float GetRadius() const;
+		RIGIDPHYSICS_API void SetRadius(float InRadius);
+
+		RIGIDPHYSICS_API FString ToString() const;
+
+	private:
+		float Radius = 0;
+	};
+
+	struct FBoxGeometry
+	{
+		FBoxGeometry() = default;
+		RIGIDPHYSICS_API FBoxGeometry(const FVector3f& InExtents);
+
+		bool operator==(const FBoxGeometry&) const = default;
+		bool operator!=(const FBoxGeometry&) const = default;
+
+		RIGIDPHYSICS_API FVector3f GetExtents() const;
+		RIGIDPHYSICS_API void SetExtents(const FVector3f& InExtents);
+
+		RIGIDPHYSICS_API FString ToString() const;
+
+	private:
+		FVector3f Extents = FVector3f::Zero();
+	};
+
+	struct FCapsuleGeometry
+	{
+		FCapsuleGeometry() = default;
+		RIGIDPHYSICS_API FCapsuleGeometry(float InRadius, float InSegmentHalfLength);
+
+		bool operator==(const FCapsuleGeometry&) const = default;
+		bool operator!=(const FCapsuleGeometry&) const = default;
+
+		RIGIDPHYSICS_API float GetRadius() const;
+		RIGIDPHYSICS_API void SetRadius(const float InRadius);
+
+		RIGIDPHYSICS_API float GetSegmentHalfLength() const;
+		RIGIDPHYSICS_API void SetSegmentHalfLength(const float InHalfLength);
+
+		RIGIDPHYSICS_API FString ToString() const;
+
+	private:
+		float Radius = 0;
+		float SegmentHalfLength = 0;
+	};
+} // namespace UE::Physics
+
+#endif // UE_RIGIDPHYSICS_API_ENABLED
diff --git a/Engine/Source/Runtime/Experimental/RigidPhysics/Public/RigidPhysics/Geometry/ConvexGeometry.h b/Engine/Source/Runtime/Experimental/RigidPhysics/Public/RigidPhysics/Geometry/ConvexGeometry.h
new file mode 100644
--- /dev/null
+++ b/Engine/Source/Runtime/Experimental/RigidPhysics/Public/RigidPhysics/Geometry/ConvexGeometry.h
@@ -0,0 +1,29 @@
+// Copyright Epic Games, Inc. All Rights Reserved.
+#pragma once
+
+#include "RigidPhysics/RigidFwd.h"
+
+#if UE_RIGIDPHYSICS_API_ENABLED
+
+#include "Chaos/Convex.h"
+
+UE_EXPERIMENTAL(5.8, "The new Chaos API is experimental")
+namespace UE::Physics
+{
+	class FConvexGeometry
+	{
+	public:
+		UE_INTERNAL FConvexGeometry() = default;
+		UE_INTERNAL RIGIDPHYSICS_API FConvexGeometry(Chaos::FImplicitConvex3* InImplicit);
+
+		bool operator==(const FConvexGeometry&) const = default;
+		bool operator!=(const FConvexGeometry&) const = default;
+
+		RIGIDPHYSICS_API FString ToString() const;
+
+	private:
+		Chaos::FConvexPtr Implicit;
+	};
+} // namespace UE::Physics
+
+#endif // UE_RIGIDPHYSICS_API_ENABLED
diff --git a/Engine/Source/Runtime/Experimental/RigidPhysics/Public/RigidPhysics/Geometry/TriangleMeshGeometry.h b/Engine/Source/Runtime/Experimental/RigidPhysics/Public/RigidPhysics/Geometry/TriangleMeshGeometry.h
new file mode 100644
--- /dev/null
+++ b/Engine/Source/Runtime/Experimental/RigidPhysics/Public/RigidPhysics/Geometry/TriangleMeshGeometry.h
@@ -0,0 +1,29 @@
+// Copyright Epic Games, Inc. All Rights Reserved.
+#pragma once
+
+#include "RigidPhysics/RigidFwd.h"
+
+#if UE_RIGIDPHYSICS_API_ENABLED
+
+#include "Chaos/TriangleMeshImplicitObject.h"
+
+UE_EXPERIMENTAL(5.8, "The new Chaos API is experimental")
+namespace UE::Physics
+{
+	class FTriangleMeshGeometry
+	{
+	public:
+		UE_INTERNAL FTriangleMeshGeometry() = default;
+		UE_INTERNAL RIGIDPHYSICS_API FTriangleMeshGeometry(Chaos::FTriangleMeshImplicitObject* InImplicit);
+
+		bool operator==(const FTriangleMeshGeometry&) const = default;
+		bool operator!=(const FTriangleMeshGeometry&) const = default;
+
+		RIGIDPHYSICS_API FString ToString() const;
+
+	private:
+		Chaos::FTriangleMeshImplicitObjectPtr Implicit;
+	};
+} // namespace UE::Physics
+
+#endif // UE_RIGIDPHYSICS_API_ENABLED
