Chaos API: Adding ShapeInstance and tests.
- Only Sphere and Boxes are currently supported.
- The API on shape instance is currently quite limited.

#Synced-CL 51413247

diff --git a/Engine/Plugins/Experimental/ChaosRigidPhysicsAsync/Source/ChaosRigidPhysicsAsync/Private/ChaosRigidPhysicsAsync/RigidBodyAsync.cpp b/Engine/Plugins/Experimental/ChaosRigidPhysicsAsync/Source/ChaosRigidPhysicsAsync/Private/ChaosRigidPhysicsAsync/RigidBodyAsync.cpp
--- a/Engine/Plugins/Experimental/ChaosRigidPhysicsAsync/Source/ChaosRigidPhysicsAsync/Private/ChaosRigidPhysicsAsync/RigidBodyAsync.cpp
+++ b/Engine/Plugins/Experimental/ChaosRigidPhysicsAsync/Source/ChaosRigidPhysicsAsync/Private/ChaosRigidPhysicsAsync/RigidBodyAsync.cpp
@@ -4,7 +4,6 @@
 #include "ChaosRigidPhysicsAsync/RigidSceneAsync.h"
 #include "Chaos/ShapeInstance.h"
 #include "RigidPhysics/RigidBody.h"
-#include "RigidPhysics/RigidShapeInstance.h"
 
 #if UE_RIGIDPHYSICS_API_ENABLED
 
@@ -172,34 +171,42 @@ namespace Chaos::Rigids::Async
 
 	int32 FRigidBodyAsyncGT::GetNumShapes() const
 	{
-		return GetProxyAPI().ShapesArray().Num();
+		return ShapeInstanceView.Num();
 	}
 
-	UE::Physics::FConstRigidShapeInstancePtr FRigidBodyAsyncGT::GetShape(int32 InShapeIndex) const
+	FRigidShapeInstanceAsync* FRigidBodyAsyncGT::GetShape(int32 InShapeIndex) const
 	{
-		using namespace UE::Physics;
+		return SceneAsync->GetShape(ShapeInstanceView, InShapeIndex);
+	}
 
-		if (GetProxyAPI().ShapesArray().IsValidIndex(InShapeIndex))
-		{
-			return FConstRigidShapeInstancePtr(GetProxyAPI().ShapesArray()[InShapeIndex].Get());
-		}
-		return FConstRigidShapeInstancePtr();
+	FRigidShapeInstanceAsync* FRigidBodyAsyncGT::CreateShape(const UE::Physics::FRigidShapeInstanceSetup& InSetup)
+	{
+		return CreateShapeInternal(FRigidShapeInstanceAsync::Create(InSetup));
 	}
 
-	void FRigidBodyAsyncGT::AddShape(const UE::Physics::FConstRigidShapeInstancePtr& InShape)
+	FRigidShapeInstanceAsync* FRigidBodyAsyncGT::CreateShape(TUniquePtr<Chaos::FPerShapeData>&& InShapeData)
 	{
-		using namespace Chaos;
+		return CreateShapeInternal(InShapeData.Release());
+	}
 
-		TArray<FImplicitObjectPtr> ImplicitsObjects;
-		ImplicitsObjects = { InShape->GetPerShapeData()->GetGeometry() };
-		GetProxyAPI().MergeGeometry(MoveTemp(ImplicitsObjects));
+	void FRigidBodyAsyncGT::DestroyShape(IRigidShapeInstance* InShape, bool bWakeTouching)
+	{
+		if (FRigidShapeInstanceAsync* CastedShape = InShape->AsAChecked<FRigidShapeInstanceAsync>())
+		{
+			// Remove the shape from our view
+			SceneAsync->RemoveShape(*this, *CastedShape);
 
-		// TODO_CHAOSAPI: We are wrapping the collider's PerShapeData in a UniquePtr here
-		// which is leaving a raw pointer to the same object behind in InShape.
-		// This is temporary until the API is used everywhere and we can refactor.
-		FShapesArray Shapes;
-		Shapes.Add(TUniquePtr<FPerShapeData>(InShape->GetPerShapeData()));
-		GetProxyAPI().MergeShapesArray(MoveTemp(Shapes));
+			// Currently, the proxy owns the shape, so also remove the memory
+			GetProxyAPI().RemoveShape(CastedShape->PerShapeData, bWakeTouching);
+			CastedShape->PerShapeData = nullptr;
+
+			SceneAsync->DestroyShape(CastedShape);
+		}
+	}
+
+	FShapeInstanceView& FRigidBodyAsyncGT::GetShapeInstanceView()
+	{
+		return ShapeInstanceView;
 	}
 
 	void FRigidBodyAsyncGT::InitTransform(const FTransform& InTransform)
@@ -311,16 +318,38 @@ namespace Chaos::Rigids::Async
 	//void SetNumVelocityIterations(int32 InNum) {  }
 	//void SetNumProjectionIterations(int32 InNum) {  }
 	//void SetIsSleeping(bool bInSleeping) {  }
-}
+	
+	FRigidShapeInstanceAsync* FRigidBodyAsyncGT::CreateShapeInternal(Chaos::FPerShapeData* InShapeData)
+	{
+		using namespace Chaos;
+		TArray<FImplicitObjectPtr> ImplicitsObjects;
+		ImplicitsObjects = { InShapeData->GetGeometry() };
+		GetProxyAPI().MergeGeometry(MoveTemp(ImplicitsObjects));
+
+		// TODO_CHAOSAPI: We are wrapping the instance's PerShapeData in a UniquePtr here
+		// which is leaving a raw pointer to the same object behind.
+		// This is temporary until the API is used everywhere and we can refactor.
+		FShapesArray Shapes;
+		Shapes.Add(TUniquePtr<FPerShapeData>(InShapeData));
+		GetProxyAPI().MergeShapesArray(MoveTemp(Shapes));
+
+		FRigidShapeInstanceAsync* ShapeInstance = SceneAsync->AllocateShape();
+		ShapeInstance->PerShapeData = InShapeData;
+		SceneAsync->AddShape(*this, *ShapeInstance);
+		return ShapeInstance;
+	}
+} // namespace Chaos::Rigids::Async
 
 namespace Chaos::Rigids::Async
 {
 	UE_RIGIDPHYSICS_RIGIDTYPED_IMPL(FRigidBodyAsyncPT, UE::Physics::IRigidBody);
 
 	FRigidBodyAsyncPT::FRigidBodyAsyncPT(
+		FRigidSceneAsyncPT* InScene,
 		const UE::Physics::FRigidBodyHandle& InBodyHandle,
 		FSingleParticlePhysicsProxy* InProxy)
-		: Handle(InBodyHandle)
+		: Scene(InScene)
+		, Handle(InBodyHandle)
 		, Proxy(InProxy)
 	{
 	}
@@ -362,24 +391,31 @@ namespace Chaos::Rigids::Async
 
 	int32 FRigidBodyAsyncPT::GetNumShapes() const
 	{
-		return GetProxyAPI()->ShapesArray().Num();
+		return ShapeInstanceView.Num();
 	}
 
-	UE::Physics::FConstRigidShapeInstancePtr FRigidBodyAsyncPT::GetShape(int32 InShapeIndex) const
+	FRigidShapeInstanceAsync* FRigidBodyAsyncPT::GetShape(int32 InShapeIndex) const
 	{
-		using namespace UE::Physics;
-
-		if (GetProxyAPI()->ShapesArray().IsValidIndex(InShapeIndex))
-		{
-			return FConstRigidShapeInstancePtr(GetProxyAPI()->ShapesArray()[InShapeIndex].Get());
-		}
+		return Scene->GetShape(ShapeInstanceView, InShapeIndex);
+	}
 
-		return FConstRigidShapeInstancePtr();
+	FShapeInstanceView& FRigidBodyAsyncPT::GetShapeInstanceView()
+	{
+		return ShapeInstanceView;
 	}
 
-	void FRigidBodyAsyncPT::AddShape(const UE::Physics::FConstRigidShapeInstancePtr& InCollider)
+	void FRigidBodyAsyncPT::UpdateShapeInstance(FRigidShapeInstanceAsync& ShapeInstance, const int32 InShapeIndex)
 	{
-		UE_RIGIDPHYSICS_CHECK(false);
+		const Chaos::FShapesArray& ShapesArray = GetProxyAPI()->ShapesArray();
+		// Note: This index isn't always valid so we have to guard. This is true if multiple RemoveShape calls happen in one tick.
+		if (ShapesArray.IsValidIndex(InShapeIndex))
+		{
+			ShapeInstance.PerShapeData = ShapesArray[InShapeIndex].Get();
+		}
+		else
+		{
+			ShapeInstance.PerShapeData = nullptr;
+		}
 	}
 
 	void FRigidBodyAsyncPT::InitTransform(const FTransform& InTransform)
@@ -469,6 +505,6 @@ namespace Chaos::Rigids::Async
 		return UnknownName;
 #endif
 	};
-}
+} // namespace Chaos::Rigids::Async
 
-#endif // UE_RIGIDPHYSICS_API_ENABLED
\ No newline at end of file
+#endif // UE_RIGIDPHYSICS_API_ENABLED
diff --git a/Engine/Plugins/Experimental/ChaosRigidPhysicsAsync/Source/ChaosRigidPhysicsAsync/Private/ChaosRigidPhysicsAsync/RigidGeometryCollectionAsync.cpp b/Engine/Plugins/Experimental/ChaosRigidPhysicsAsync/Source/ChaosRigidPhysicsAsync/Private/ChaosRigidPhysicsAsync/RigidGeometryCollectionAsync.cpp
--- a/Engine/Plugins/Experimental/ChaosRigidPhysicsAsync/Source/ChaosRigidPhysicsAsync/Private/ChaosRigidPhysicsAsync/RigidGeometryCollectionAsync.cpp
+++ b/Engine/Plugins/Experimental/ChaosRigidPhysicsAsync/Source/ChaosRigidPhysicsAsync/Private/ChaosRigidPhysicsAsync/RigidGeometryCollectionAsync.cpp
@@ -63,11 +63,9 @@ namespace Chaos::Rigids::Async
 		return 0;
 	}
 
-	UE::Physics::FConstRigidShapeInstancePtr FRigidGeometryCollectionBodyAsyncGT::GetShape(int32 InShapeIndex) const
+	UE::Physics::IRigidShapeInstance* FRigidGeometryCollectionBodyAsyncGT::GetShape(int32 InShapeIndex) const
 	{
-		using namespace UE::Physics;
-
-		return FConstRigidShapeInstancePtr();
+		return nullptr;
 	}
 }
 
@@ -130,11 +128,9 @@ namespace Chaos::Rigids::Async
 		return 0;
 	}
 
-	UE::Physics::FConstRigidShapeInstancePtr FRigidGeometryCollectionBodyAsyncPT::GetShape(int32 InShapeIndex) const
+	UE::Physics::IRigidShapeInstance* FRigidGeometryCollectionBodyAsyncPT::GetShape(int32 InShapeIndex) const
 	{
-		using namespace UE::Physics;
-
-		return FConstRigidShapeInstancePtr();
+		return nullptr;
 	}
 }
 
diff --git a/Engine/Plugins/Experimental/ChaosRigidPhysicsAsync/Source/ChaosRigidPhysicsAsync/Private/ChaosRigidPhysicsAsync/RigidSceneAsync.cpp b/Engine/Plugins/Experimental/ChaosRigidPhysicsAsync/Source/ChaosRigidPhysicsAsync/Private/ChaosRigidPhysicsAsync/RigidSceneAsync.cpp
--- a/Engine/Plugins/Experimental/ChaosRigidPhysicsAsync/Source/ChaosRigidPhysicsAsync/Private/ChaosRigidPhysicsAsync/RigidSceneAsync.cpp
+++ b/Engine/Plugins/Experimental/ChaosRigidPhysicsAsync/Source/ChaosRigidPhysicsAsync/Private/ChaosRigidPhysicsAsync/RigidSceneAsync.cpp
@@ -265,6 +265,9 @@ namespace Chaos::Rigids::Async
 					});
 			}
 
+			// If the body had any shapes left, make sure to free them all.
+			ShapeInstanceContainer.Free(Body->GetShapeInstanceView());
+
 			Body->Deactivate();
 			UnRegisterBody(BodyHandle, Body);
 			BodyPool.Free(Body);
@@ -309,7 +312,62 @@ namespace Chaos::Rigids::Async
 			delete GC;
 		}
 	}
-	
+
+	FRigidShapeInstanceAsync* FRigidSceneAsyncGT::AllocateShape()
+	{
+		return ShapeInstanceContainer.Allocate(SceneId);
+	}
+
+	void FRigidSceneAsyncGT::DestroyShape(FRigidShapeInstanceAsync* InShape)
+	{
+		ShapeInstanceContainer.Free(InShape);
+	}
+
+	FRigidShapeInstanceAsync* FRigidSceneAsyncGT::GetShape(const FRigidObjectId& InId)  const
+	{
+		return ShapeInstanceContainer.Find(InId);
+	}
+
+	void FRigidSceneAsyncGT::AddShape(FRigidBodyAsyncGT& Body, FRigidShapeInstanceAsync& InShapeInstance)
+	{
+		const FRigidBodyHandle BodyHandle = Body.GetHandle();
+		const FRigidShapeInstanceHandle ShapeInstanceHandle = InShapeInstance.GetHandle();
+		FShapeInstanceView& View = Body.GetShapeInstanceView();
+		const int32 ShapeIndex = ShapeInstanceContainer.Add(InShapeInstance, View);
+
+		if (FPBDRigidsSolver* Solver = GetSolver())
+		{
+			Solver->EnqueueCommandImmediate(
+				[this, BodyHandle, ShapeInstanceHandle, ShapeIndex]()
+				{
+					ScenePT->AddShapePT(BodyHandle, ShapeInstanceHandle, ShapeIndex);
+				});
+		}
+	}
+
+	void FRigidSceneAsyncGT::RemoveShape(FRigidBodyAsyncGT& Body, FRigidShapeInstanceAsync& InShapeInstance)
+	{
+		const FRigidBodyHandle BodyHandle = Body.GetHandle();
+		FShapeInstanceView& View = Body.GetShapeInstanceView();
+		FRigidShapeInstanceHandle ShapeInstanceHandle = InShapeInstance.GetHandle();
+
+		if (FPBDRigidsSolver* Solver = GetSolver())
+		{
+			Solver->EnqueueCommandImmediate(
+				[this, BodyHandle, ShapeInstanceHandle]()
+				{
+					ScenePT->RemoveShapePT(BodyHandle, ShapeInstanceHandle);
+				});
+		}
+
+		ShapeInstanceContainer.Remove(InShapeInstance, View);
+	}
+
+	FRigidShapeInstanceAsync* FRigidSceneAsyncGT::GetShape(const FShapeInstanceView& InView, const int32 InIndex) const
+	{
+		return ShapeInstanceContainer.Get(InView, InIndex);
+	}
+
 	FJointConstraint6DOFAsyncGT* FRigidSceneAsyncGT::CreateJointConstraint6DOF()
 	{
 		FJointConstraint6DOFAsyncGT* Constraint = new FJointConstraint6DOFAsyncGT(this);
@@ -645,7 +703,7 @@ namespace Chaos::Rigids::Async
 		CheckSimContext();
 		UE_RIGIDPHYSICS_CHECK(GetBody(InBodyHandle.GetId()) == nullptr);
 
-		FRigidBodyAsyncPT* Body = BodyPool.Alloc(InBodyHandle, InProxy);
+		FRigidBodyAsyncPT* Body = BodyPool.Alloc(this, InBodyHandle, InProxy);
 		RegisterBody(InBodyHandle, Body);
 	}
 
@@ -657,6 +715,8 @@ namespace Chaos::Rigids::Async
 		UnRegisterBody(InBodyHandle, Body);
 		if (FRigidBodyAsyncPT* TypedBody = Body->AsAChecked<FRigidBodyAsyncPT>())
 		{
+			ShapeInstanceContainer.Free(TypedBody->GetShapeInstanceView());
+
 			BodyPool.Free(TypedBody);
 		}
 	}
@@ -710,6 +770,68 @@ namespace Chaos::Rigids::Async
 
 		return nullptr;
 	}
+
+	FRigidShapeInstanceAsync* FRigidSceneAsyncPT::GetShape(const FRigidObjectId& InId) const
+	{
+		CheckSimContext();
+
+		return ShapeInstanceContainer.Find(InId);
+	}
+
+	void FRigidSceneAsyncPT::AddShapePT(const FRigidBodyHandle BodyHandle, const FRigidShapeInstanceHandle ShapeInstanceHandle, const int32 ShapeIndex)
+	{
+		CheckSimContext();
+
+		if (UE::Physics::IRigidBody* Body = GetBody(BodyHandle.GetId()))
+		{
+			if (FRigidBodyAsyncPT* TypedBody = Body->AsAChecked<FRigidBodyAsyncPT>())
+			{
+				if (TypedBody->IsValid())
+				{
+					FRigidShapeInstanceAsync* ShapeInstance = ShapeInstanceContainer.Allocate(ShapeInstanceHandle);
+					FShapeInstanceView& View = TypedBody->GetShapeInstanceView();
+
+					ShapeInstanceContainer.Add(*ShapeInstance, ShapeIndex, View);
+					TypedBody->UpdateShapeInstance(*ShapeInstance, ShapeIndex);
+				}
+			}
+		}
+	}
+
+	void FRigidSceneAsyncPT::RemoveShapePT(const FRigidBodyHandle BodyHandle, const FRigidShapeInstanceHandle ShapeInstanceHandle)
+	{
+		CheckSimContext();
+
+		if (UE::Physics::IRigidBody* Body = GetBody(BodyHandle.GetId()))
+		{
+			if (FRigidBodyAsyncPT* TypedBody = Body->AsAChecked<FRigidBodyAsyncPT>())
+			{
+				if (TypedBody->IsValid())
+				{
+					FShapeInstanceView& View = TypedBody->GetShapeInstanceView();
+					FRigidShapeInstanceAsync* ShapeInstance = GetShape(ShapeInstanceHandle.GetId());
+
+					// Remove from the body's shape view
+					ShapeInstanceContainer.Remove(*ShapeInstance, View);
+
+					// Update all shape instances to point at the correct shape data (PT rebuilds from "scratch" so old data gets invalidated)
+					for (int32 I = 0; I < View.Num(); ++I)
+					{
+						TypedBody->UpdateShapeInstance(*GetShape(View, I), I);
+					}
+					// Free the shape
+					ShapeInstanceContainer.Free(ShapeInstance);
+				}
+			}
+		}
+	}
+
+	FRigidShapeInstanceAsync* FRigidSceneAsyncPT::GetShape(const FShapeInstanceView& InView, const int32 InIndex) const
+	{
+		CheckSimContext();
+
+		return ShapeInstanceContainer.Get(InView, InIndex);
+	}
 	
 	FJointConstraint6DOFAsyncPT* FRigidSceneAsyncPT::CreateJointConstraint6DOFPT(const FJointConstraint6DOFHandle& InHandle, Chaos::FJointConstraint* InJointConstraint)
 	{
diff --git a/Engine/Plugins/Experimental/ChaosRigidPhysicsAsync/Source/ChaosRigidPhysicsAsync/Private/ChaosRigidPhysicsAsync/RigidShapeInstanceAsync.cpp b/Engine/Plugins/Experimental/ChaosRigidPhysicsAsync/Source/ChaosRigidPhysicsAsync/Private/ChaosRigidPhysicsAsync/RigidShapeInstanceAsync.cpp
new file mode 100644
--- /dev/null
+++ b/Engine/Plugins/Experimental/ChaosRigidPhysicsAsync/Source/ChaosRigidPhysicsAsync/Private/ChaosRigidPhysicsAsync/RigidShapeInstanceAsync.cpp
@@ -0,0 +1,179 @@
+// Copyright Epic Games, Inc. All Rights Reserved.
+#include "ChaosRigidPhysicsAsync/RigidShapeInstanceAsync.h"
+
+
+#if UE_RIGIDPHYSICS_API_ENABLED
+
+#include "Chaos/Convex.h"
+#include "Chaos/ShapeInstance.h"
+#include "Chaos/TriangleMeshImplicitObject.h"
+
+#include "Physics/Experimental/ChaosInterfaceUtils.h"
+#include "PhysicsEngine/AggregateGeom.h"
+
+namespace Chaos::Rigids::Async
+{
+	UE_RIGIDPHYSICS_RIGIDTYPED_IMPL(FRigidShapeInstanceAsync, UE::Physics::IRigidShapeInstance);
+
+	Chaos::FPerShapeData* FRigidShapeInstanceAsync::Create(const UE::Physics::FRigidShapeInstanceSetup& InSetup)
+	{
+		const UE::Physics::FAnyGeometry& InGeometry = InSetup.Geometry;
+		FKAggregateGeom AggGeom;
+		InGeometry.Visit([&AggGeom]<typename GeometryType>(const GeometryType& Geometry)
+		{
+			if constexpr (std::is_same_v<GeometryType, UE::Physics::FBoxGeometry>)
+			{
+				const FVector3f Extents = Geometry.GetExtents();
+				FKBoxElem BoxElem = FKBoxElem(Extents.X, Extents.Y, Extents.Z);
+				BoxElem.Center = FVector::Zero();
+				AggGeom.AddElement(BoxElem);
+			}
+			else if constexpr (std::is_same_v<GeometryType, UE::Physics::FSphereGeometry>)
+			{
+				FKSphereElem SphereElem = FKSphereElem(Geometry.GetRadius());
+				SphereElem.Center = FVector::Zero();
+				AggGeom.AddElement(SphereElem);
+			}
+			else if constexpr (std::is_same_v<GeometryType, UE::Physics::FCapsuleGeometry>)
+			{
+				ensureMsgf(false, TEXT("TODO"));
+			}
+			else if constexpr (std::is_same_v<GeometryType, UE::Physics::FConvexGeometry>)
+			{
+				ensureMsgf(false, TEXT("TODO"));
+			}
+			else if constexpr (std::is_same_v<GeometryType, UE::Physics::FTriangleMeshGeometry>)
+			{
+				ensureMsgf(false, TEXT("TODO"));
+			}
+			else
+			{
+				static_assert(false, "Unsupported shape type");
+			}
+		});
+
+		// Not used - filters are set up after
+		FBodyCollisionData BodyCollisionData;
+
+		FGeometryAddParams AddParams;
+		AddParams.bDoubleSided = false;
+		AddParams.CollisionData = BodyCollisionData;
+		AddParams.CollisionTraceType = ECollisionTraceFlag::CTF_UseSimpleAsComplex;
+		AddParams.Scale = FVector3d::One();
+		AddParams.SimpleMaterial = nullptr;
+		AddParams.ComplexMaterials = {};
+		AddParams.ComplexMaterialMasks = {};
+		AddParams.LocalTransform = FTransform(InSetup.LocalTransform);
+		AddParams.WorldTransform = FTransform::Identity;
+		AddParams.Geometry = &AggGeom;
+
+		TArray<Chaos::FImplicitObjectPtr> Geoms;
+		Chaos::FShapesArray Shapes;
+		ChaosInterface::CreateGeometry(AddParams, Geoms, Shapes);
+
+		if (Shapes.Num() != 1)
+		{
+			UE_LOGFMT(LogRigidPhysics, Warning, "MakeShapeInstanceFromAddParams is intended to be used for single shapes only. Have {0} shapes.", Shapes.Num());
+		}
+
+		Chaos::FPerShapeData* PerShapeData = nullptr;
+		if (Shapes.Num() > 0)
+		{
+			PerShapeData = Shapes[0].Release();
+		}
+
+		if (PerShapeData != nullptr)
+		{
+			const Chaos::EFilterFlags FilterFlags = Chaos::EFilterFlags::SimpleCollision | Chaos::EFilterFlags::ContactNotify | Chaos::EFilterFlags::ModifyContacts;
+			const Chaos::Filter::FShapeFilterData ShapeFilter = Chaos::Filter::FShapeFilterBuilder::BuildBlockAll(FilterFlags);
+			PerShapeData->SetQueryEnabled(true);
+			PerShapeData->SetSimEnabled(true);
+			PerShapeData->SetShapeFilterData(ShapeFilter);
+		}
+		return PerShapeData;
+	}
+
+	UE::Physics::FRigidShapeInstanceHandle FRigidShapeInstanceAsync::GetHandle() const
+	{
+		return Handle;
+	}
+
+	UE::Physics::FAnyGeometry FRigidShapeInstanceAsync::GetGeometry() const
+	{
+		using FImplicitTransformed = TImplicitObjectTransformed<FReal, 3>;
+		using FBoxGeometry = UE::Physics::FBoxGeometry;
+		using FSphereGeometry = UE::Physics::FSphereGeometry;
+		using FAnyGeometry = UE::Physics::FAnyGeometry;
+
+		const FImplicitObject* ImplicitGeom = PerShapeData->GetGeometry();
+		EImplicitObjectType ImplicitType = ImplicitGeom->GetType();
+		// TODO: Looks like there's other "wrapper" types like instanced and scaled that should be handled.
+		if (ImplicitType == ImplicitObjectType::Transformed)
+		{
+			const FImplicitTransformed& Transformed = ImplicitGeom->GetObjectChecked<FImplicitTransformed>();
+			ImplicitGeom = Transformed.GetTransformedObject();
+			ImplicitType = ImplicitGeom->GetType();
+		}
+
+		switch (ImplicitType)
+		{
+			case ImplicitObjectType::Sphere:
+			{
+				return FSphereGeometry(ImplicitGeom->GetObjectChecked<FImplicitSphere3>().GetRadiusf());
+			}
+			case ImplicitObjectType::Box:
+			{
+				return FBoxGeometry(FVector3f(ImplicitGeom->GetObjectChecked<FImplicitBox3>().Extents()));
+			}
+			default:
+			{
+				ensureMsgf(false, TEXT("TODO"));
+				break;
+			}
+		}
+		return FAnyGeometry();
+	}
+
+	FTransform3f FRigidShapeInstanceAsync::GetLocalTransform() const
+	{
+		using FImplicitTransformed = TImplicitObjectTransformed<FReal, 3>;
+
+		const FImplicitObjectRef Geom = PerShapeData->GetGeometry();
+		switch (Geom->GetType())
+		{
+			case ImplicitObjectType::Sphere:
+			{
+				FImplicitSphere3& Sphere = Geom->GetObjectChecked<FImplicitSphere3>();
+				return FTransform3f(Sphere.GetCenterf());
+			}
+			case ImplicitObjectType::Box:
+			{
+				FImplicitBox3& Box = Geom->GetObjectChecked<FImplicitBox3>();
+				return FTransform3f(FVector3f(Box.Center()));
+			}
+			case ImplicitObjectType::Transformed:
+			{
+				const FImplicitTransformed& Transformed = Geom->GetObjectChecked<FImplicitTransformed>();
+				return FTransform3f(Transformed.GetTransform());
+			}
+			default:
+			{
+				ensureMsgf(false, TEXT("TODO"));
+				break;
+			}
+		}
+		return FTransform3f();
+	}
+
+	Chaos::Filter::FShapeFilterData FRigidShapeInstanceAsync::GetShapeFilter() const
+	{
+		return PerShapeData->GetShapeFilterData();
+	}
+
+	void FRigidShapeInstanceAsync::SetShapeFilter(const FShapeFilterData& ShapeFilter)
+	{
+		PerShapeData->SetShapeFilterData(ShapeFilter);
+	}
+} // namespace Chaos::Rigids::Async
+
+#endif // UE_RIGIDPHYSICS_API_ENABLED
diff --git a/Engine/Plugins/Experimental/ChaosRigidPhysicsAsync/Source/ChaosRigidPhysicsAsync/Private/ChaosRigidPhysicsAsync/ShapeInstanceContainerAsync.cpp b/Engine/Plugins/Experimental/ChaosRigidPhysicsAsync/Source/ChaosRigidPhysicsAsync/Private/ChaosRigidPhysicsAsync/ShapeInstanceContainerAsync.cpp
new file mode 100644
--- /dev/null
+++ b/Engine/Plugins/Experimental/ChaosRigidPhysicsAsync/Source/ChaosRigidPhysicsAsync/Private/ChaosRigidPhysicsAsync/ShapeInstanceContainerAsync.cpp
@@ -0,0 +1,161 @@
+// Copyright Epic Games, Inc. All Rights Reserved.
+#include "ChaosRigidPhysicsAsync/ShapeInstanceContainerAsync.h"
+
+#if UE_RIGIDPHYSICS_API_ENABLED
+
+#include "RigidPhysics/RigidLog.h"
+
+namespace Chaos::Rigids::Async
+{
+	int32 FShapeInstanceView::Num() const
+	{
+		return Ids.Num();
+	}
+
+	FShapeInstanceContainerAsyncGT::FShapeInstanceContainerAsyncGT()
+		: Pool(10)
+		, Registry(10)
+	{
+	}
+
+	FRigidShapeInstanceAsync* FShapeInstanceContainerAsyncGT::Allocate(const FRigidSceneId InSceneId)
+	{
+		FRigidShapeInstanceAsync* Result = Pool.Alloc();
+		const FRigidObjectId Id = Registry.AddObject(Result);
+		Result->Handle = UE::Physics::FRigidShapeInstanceHandle(InSceneId, Id);
+		return Result;
+	}
+
+	void FShapeInstanceContainerAsyncGT::Free(FRigidShapeInstanceAsync* InInstance)
+	{
+		Registry.RemoveObjectChecked(InInstance->Handle.GetId(), InInstance);
+		Pool.Free(InInstance);
+	}
+
+	FRigidShapeInstanceAsync* FShapeInstanceContainerAsyncGT::Find(const FRigidObjectId& InId) const
+	{
+		return Registry.FindObject(InId);
+	}
+
+	int32 FShapeInstanceContainerAsyncGT::Add(FRigidShapeInstanceAsync& InInstance, FShapeInstanceView& InOutView)
+	{
+		UE_RIGIDPHYSICS_CHECK(InInstance.bAttachedToBody == false);
+		InInstance.bAttachedToBody = true;
+		InOutView.Ids.Add(InInstance.Handle.GetId());
+		return InOutView.Ids.Num() - 1;
+	}
+
+	void FShapeInstanceContainerAsyncGT::Remove(FRigidShapeInstanceAsync& InInstance, FShapeInstanceView& InOutView)
+	{
+		UE_RIGIDPHYSICS_CHECK(InInstance.bAttachedToBody == true);
+		InInstance.bAttachedToBody = false;
+		InOutView.Ids.Remove(InInstance.Handle.GetId());
+	}
+
+	FRigidShapeInstanceAsync* FShapeInstanceContainerAsyncGT::Get(const FShapeInstanceView& InView, const int32 InIndex) const
+	{
+		const bool bIsValidIndex = InView.Ids.IsValidIndex(InIndex);
+		ensure(bIsValidIndex);
+		return bIsValidIndex ? Find(InView.Ids[InIndex]) : nullptr;
+	}
+
+	void FShapeInstanceContainerAsyncGT::Free(FShapeInstanceView& InView)
+	{
+		for (int32 I = 0; I < InView.Num(); ++I)
+		{
+			FRigidShapeInstanceAsync* ShapeInstance = Find(InView.Ids[I]);
+			UE_RIGIDPHYSICS_CHECK(ShapeInstance);
+			Free(ShapeInstance);
+		}
+		InView.Ids.Reset();
+	}
+
+	FShapeInstanceContainerAsyncPT::FShapeInstanceContainerAsyncPT()
+		: Pool(10)
+	{
+	}
+
+	FRigidShapeInstanceAsync* FShapeInstanceContainerAsyncPT::Allocate(const FRigidShapeInstanceHandle& InHandle)
+	{
+		UE_RIGIDPHYSICS_CHECK(InHandle.GetId().IsValid());
+
+		FRigidShapeInstanceAsync* Result = Pool.Alloc();
+		Result->Handle = InHandle;
+
+		const FRigidObjectId& Id = InHandle.GetId();
+		const int32 ArrayIndex = Id.GetArrayIndex();
+		Registry.SetNum(ArrayIndex + 1);
+		FRegistryItem& Item = Registry[ArrayIndex];
+		// Make sure the slot is empty. If it isn't that means some GT->PT communication issue happened.
+		UE_RIGIDPHYSICS_CHECK(Item.Object == nullptr);
+		Item.Object = Result;
+		Item.Epoch = Id.GetEpoch();
+
+		return Result;
+	}
+
+	void FShapeInstanceContainerAsyncPT::Free(FRigidShapeInstanceAsync* InInstance)
+	{
+		const FRigidShapeInstanceHandle& Handle = InInstance->Handle;
+		const FRigidObjectId& Id = Handle.GetId();
+		const int32 ArrayIndex = Id.GetArrayIndex();
+
+		UE_RIGIDPHYSICS_CHECK(Registry.IsValidIndex(ArrayIndex));
+		FRegistryItem& Item = Registry[ArrayIndex];
+		UE_RIGIDPHYSICS_CHECK(Item.Object == InInstance);
+		UE_RIGIDPHYSICS_CHECK(Item.Epoch == Id.GetEpoch());
+		Item = FRegistryItem();
+
+		Pool.Free(InInstance);
+	}
+
+	FRigidShapeInstanceAsync* FShapeInstanceContainerAsyncPT::Find(const FRigidObjectId& InId) const
+	{
+		const int32 ArrayIndex = InId.GetArrayIndex();
+		if (Registry.IsValidIndex(ArrayIndex))
+		{
+			const FRegistryItem& Item = Registry[ArrayIndex];
+			if (Item.Epoch == InId.GetEpoch())
+			{
+				return Item.Object;
+			}
+		}
+		return nullptr;
+	}
+
+	void FShapeInstanceContainerAsyncPT::Add(FRigidShapeInstanceAsync& InInstance, const int32 InShapeIndex, FShapeInstanceView& InOutView)
+	{
+		// All adds should be happening sequentially from the GT. That means each add should always have the index that is the next to add.
+		UE_RIGIDPHYSICS_CHECK(InOutView.Num() == InShapeIndex);
+		UE_RIGIDPHYSICS_CHECK(InInstance.bAttachedToBody == false);
+		InInstance.bAttachedToBody = true;
+		InOutView.Ids.Add(InInstance.Handle.GetId());
+	}
+
+	void FShapeInstanceContainerAsyncPT::Remove(FRigidShapeInstanceAsync& InInstance, FShapeInstanceView& InOutView)
+	{
+		UE_RIGIDPHYSICS_CHECK(InInstance.bAttachedToBody == true);
+		InInstance.bAttachedToBody = false;
+		InOutView.Ids.Remove(InInstance.Handle.GetId());
+	}
+
+	FRigidShapeInstanceAsync* FShapeInstanceContainerAsyncPT::Get(const FShapeInstanceView& InView, const int32 InIndex) const
+	{
+		const bool bIsValidIndex = InView.Ids.IsValidIndex(InIndex);
+		ensure(bIsValidIndex);
+		return bIsValidIndex ? Find(InView.Ids[InIndex]) : nullptr;
+	}
+
+	void FShapeInstanceContainerAsyncPT::Free(FShapeInstanceView& InView)
+	{
+		for (int32 I = 0; I < InView.Num(); ++I)
+		{
+			FRigidShapeInstanceAsync* ShapeInstance = Find(InView.Ids[I]);
+			check(ShapeInstance);
+			Free(ShapeInstance);
+		}
+		InView.Ids.Reset();
+	}
+} // namespace Chaos::Rigids::Async
+
+#endif // UE_RIGIDPHYSICS_API_ENABLED
diff --git a/Engine/Plugins/Experimental/ChaosRigidPhysicsAsync/Source/ChaosRigidPhysicsAsync/Public/ChaosRigidPhysicsAsync/RigidBodyAsync.h b/Engine/Plugins/Experimental/ChaosRigidPhysicsAsync/Source/ChaosRigidPhysicsAsync/Public/ChaosRigidPhysicsAsync/RigidBodyAsync.h
--- a/Engine/Plugins/Experimental/ChaosRigidPhysicsAsync/Source/ChaosRigidPhysicsAsync/Public/ChaosRigidPhysicsAsync/RigidBodyAsync.h
+++ b/Engine/Plugins/Experimental/ChaosRigidPhysicsAsync/Source/ChaosRigidPhysicsAsync/Public/ChaosRigidPhysicsAsync/RigidBodyAsync.h
@@ -11,6 +11,8 @@
 #include "Physics/Experimental/PhysScene_Chaos.h"
 #include "PhysicsProxy/SingleParticlePhysicsProxy.h"
 #include "RigidPhysics/RigidBody.h"
+#include "ChaosRigidPhysicsAsync/RigidShapeInstanceAsync.h"
+#include "ChaosRigidPhysicsAsync/ShapeInstanceContainerAsync.h"
 
 namespace Chaos::Rigids::Async
 {
@@ -51,9 +53,14 @@ namespace Chaos::Rigids::Async
 		UE_INTERNAL virtual bool IsKinematic() const override final;
 		UE_INTERNAL virtual bool IsDynamic() const override final;
 		UE_INTERNAL virtual bool IsSleeping() const override final;
+
 		UE_INTERNAL virtual int32 GetNumShapes() const override final;
-		UE_INTERNAL virtual UE::Physics::FConstRigidShapeInstancePtr GetShape(int32 InShapeIndex) const override final;
-		UE_INTERNAL virtual void AddShape(const UE::Physics::FConstRigidShapeInstancePtr& InCollider) override final;
+		UE_INTERNAL virtual FRigidShapeInstanceAsync* GetShape(int32 InShapeIndex) const override final;
+		UE_INTERNAL virtual FRigidShapeInstanceAsync* CreateShape(const UE::Physics::FRigidShapeInstanceSetup& InSetup) override final;
+		UE_INTERNAL virtual FRigidShapeInstanceAsync* CreateShape(TUniquePtr<Chaos::FPerShapeData>&& InShapeData) override final;
+		UE_INTERNAL virtual void DestroyShape(IRigidShapeInstance* InShape, bool bWakeTouching) override final;
+		UE_INTERNAL FShapeInstanceView& GetShapeInstanceView();
+
 		UE_INTERNAL virtual void InitTransform(const FTransform& InTransform) override final;
 		UE_INTERNAL virtual void UpdateTransform(const FTransform& InTransform) override final;
 		UE_INTERNAL virtual void SetKinematicTarget(const FTransform& InTransform) override final;
@@ -90,9 +97,12 @@ namespace Chaos::Rigids::Async
 
 
 	private:
+		FRigidShapeInstanceAsync* CreateShapeInternal(Chaos::FPerShapeData* InShapeData);
+
 		UE::Physics::FRigidObjectId BodyId;
 		FRigidSceneAsyncGT* SceneAsync = nullptr;
 		FSingleParticlePhysicsProxy* Proxy;
+		FShapeInstanceView ShapeInstanceView;
 	};
 
 
@@ -103,7 +113,7 @@ namespace Chaos::Rigids::Async
 		{
 			return Proxy->GetPhysicsThreadAPI();
 		}
-		
+
 		inline const FRigidBodyHandle_Internal* GetProxyAPI() const
 		{
 			return Proxy->GetPhysicsThreadAPI();
@@ -113,6 +123,7 @@ namespace Chaos::Rigids::Async
 		UE_RIGIDPHYSICS_RIGIDTYPED_DECL(CHAOSRIGIDPHYSICSASYNC_API, FRigidBodyAsyncPT);
 
 		UE_INTERNAL CHAOSRIGIDPHYSICSASYNC_API FRigidBodyAsyncPT(
+			FRigidSceneAsyncPT* InScene,
 			const UE::Physics::FRigidBodyHandle& InBodyHandle,
 			FSingleParticlePhysicsProxy* InProxy);
 
@@ -126,9 +137,12 @@ namespace Chaos::Rigids::Async
 		UE_INTERNAL virtual bool IsKinematic() const override final;
 		UE_INTERNAL virtual bool IsDynamic() const override final;
 		UE_INTERNAL virtual bool IsSleeping() const override final;
+
 		UE_INTERNAL virtual int32 GetNumShapes() const override final;
-		UE_INTERNAL virtual UE::Physics::FConstRigidShapeInstancePtr GetShape(int32 InShapeIndex) const override final;
-		UE_INTERNAL virtual void AddShape(const UE::Physics::FConstRigidShapeInstancePtr& InCollider) override final;
+		UE_INTERNAL virtual FRigidShapeInstanceAsync* GetShape(int32 InShapeIndex) const override final;
+		UE_INTERNAL FShapeInstanceView& GetShapeInstanceView();
+		UE_INTERNAL void UpdateShapeInstance(FRigidShapeInstanceAsync& ShapeInstance, const int32 InShapeIndex);
+
 		UE_INTERNAL virtual void InitTransform(const FTransform& InTransform) override final;
 		UE_INTERNAL virtual void UpdateTransform(const FTransform& InTransform) override final;
 		UE_INTERNAL virtual void SetKinematicTarget(const FTransform& InTransform) override final;
@@ -146,10 +160,12 @@ namespace Chaos::Rigids::Async
 		UE_INTERNAL void SetTransform(const FTransform& InTransform, bool bIsTeleport);
 		const FString GetDebugName() const;
 
+		FRigidSceneAsyncPT* Scene = nullptr;
 		UE::Physics::FRigidBodyHandle Handle;
 		FSingleParticlePhysicsProxy* Proxy;
+		FShapeInstanceView ShapeInstanceView;
 	};
 
-}
+} // namespace Chaos::Rigids::Async
 
-#endif // UE_RIGIDPHYSICS_API_ENABLED
\ No newline at end of file
+#endif // UE_RIGIDPHYSICS_API_ENABLED
diff --git a/Engine/Plugins/Experimental/ChaosRigidPhysicsAsync/Source/ChaosRigidPhysicsAsync/Public/ChaosRigidPhysicsAsync/RigidFwdAsync.h b/Engine/Plugins/Experimental/ChaosRigidPhysicsAsync/Source/ChaosRigidPhysicsAsync/Public/ChaosRigidPhysicsAsync/RigidFwdAsync.h
--- a/Engine/Plugins/Experimental/ChaosRigidPhysicsAsync/Source/ChaosRigidPhysicsAsync/Public/ChaosRigidPhysicsAsync/RigidFwdAsync.h
+++ b/Engine/Plugins/Experimental/ChaosRigidPhysicsAsync/Source/ChaosRigidPhysicsAsync/Public/ChaosRigidPhysicsAsync/RigidFwdAsync.h
@@ -8,7 +8,7 @@
 namespace Chaos::Rigids::Async
 {
 	class FRigidSceneSettingsAsync;
-}
+} // namespace Chaos::Rigids::Async
 
 // Internal fwd decls
 namespace Chaos::Rigids::Async
@@ -23,13 +23,17 @@ namespace Chaos::Rigids::Async
 	class UE_INTERNAL FRigidGeometryCollectionAsyncPT;
 	class UE_INTERNAL FRigidSceneAsyncGT;
 	class UE_INTERNAL FRigidSceneAsyncPT;
+	class UE_INTERNAL FRigidShapeInstanceAsync;
 	
 	using UE_INTERNAL FJointConstraintHandle = UE::Physics::FJointConstraintHandle;
 	using UE_INTERNAL FJointConstraint6DOFHandle = UE::Physics::FJointConstraint6DOFHandle;
 	using UE_INTERNAL FJointConstraintRegistry = UE::Physics::FJointConstraintRegistry;
+	using UE_INTERNAL FRigidBodyHandle = UE::Physics::FRigidBodyHandle;
 	using UE_INTERNAL FRigidObjectId = UE::Physics::FRigidObjectId;
+	using UE_INTERNAL FRigidShapeInstanceHandle = UE::Physics::FRigidShapeInstanceHandle;
 	using UE_INTERNAL IJointConstraint = UE::Physics::IJointConstraint;
 	using UE_INTERNAL IJointConstraint6DOF = UE::Physics::IJointConstraint6DOF;
-}
+	using UE_INTERNAL IRigidShapeInstance = UE::Physics::IRigidShapeInstance;
+} // namespace Chaos::Rigids::Async
 
 #endif // UE_RIGIDPHYSICS_API_ENABLED
diff --git a/Engine/Plugins/Experimental/ChaosRigidPhysicsAsync/Source/ChaosRigidPhysicsAsync/Public/ChaosRigidPhysicsAsync/RigidGeometryCollectionAsync.h b/Engine/Plugins/Experimental/ChaosRigidPhysicsAsync/Source/ChaosRigidPhysicsAsync/Public/ChaosRigidPhysicsAsync/RigidGeometryCollectionAsync.h
--- a/Engine/Plugins/Experimental/ChaosRigidPhysicsAsync/Source/ChaosRigidPhysicsAsync/Public/ChaosRigidPhysicsAsync/RigidGeometryCollectionAsync.h
+++ b/Engine/Plugins/Experimental/ChaosRigidPhysicsAsync/Source/ChaosRigidPhysicsAsync/Public/ChaosRigidPhysicsAsync/RigidGeometryCollectionAsync.h
@@ -32,8 +32,8 @@ namespace Chaos::Rigids::Async
 		UE_INTERNAL virtual bool IsDynamic() const override final;
 		UE_INTERNAL virtual bool IsSleeping() const override final;
 		UE_INTERNAL virtual int32 GetNumShapes() const override final;
-		UE_INTERNAL virtual UE::Physics::FConstRigidShapeInstancePtr GetShape(int32 InShapeIndex) const override final;
-		//UE_INTERNAL virtual void AddShape(const UE::Physics::FConstRigidShapeInstancePtr& InCollider) override final;
+		UE_INTERNAL virtual UE::Physics::IRigidShapeInstance* GetShape(int32 InShapeIndex) const override final;
+		//UE_INTERNAL virtual void AddShape(const UE::Physics::IRigidShapeInstance* InCollider) override final;
 		//UE_INTERNAL virtual void InitTransform(const FTransform& InTransform) override final;
 		//UE_INTERNAL virtual void UpdateTransform(const FTransform& InTransform) override final;
 		//UE_INTERNAL virtual void SetKinematicTarget(const FTransform& InTransform) override final;
@@ -84,8 +84,8 @@ namespace Chaos::Rigids::Async
 		UE_INTERNAL virtual bool IsDynamic() const override final;
 		UE_INTERNAL virtual bool IsSleeping() const override final;
 		UE_INTERNAL virtual int32 GetNumShapes() const override final;
-		UE_INTERNAL virtual UE::Physics::FConstRigidShapeInstancePtr GetShape(int32 InShapeIndex) const override final;
-		//UE_INTERNAL virtual void AddShape(const UE::Physics::FConstRigidShapeInstancePtr& InCollider) override final;
+		UE_INTERNAL virtual UE::Physics::IRigidShapeInstance* GetShape(int32 InShapeIndex) const override final;
+		//UE_INTERNAL virtual void AddShape(const UE::Physics::IRigidShapeInstance* InCollider) override final;
 		//UE_INTERNAL virtual void InitTransform(const FTransform& InTransform) override final;
 		//UE_INTERNAL virtual void UpdateTransform(const FTransform& InTransform) override final;
 		//UE_INTERNAL virtual void SetKinematicTarget(const FTransform& InTransform) override final;
diff --git a/Engine/Plugins/Experimental/ChaosRigidPhysicsAsync/Source/ChaosRigidPhysicsAsync/Public/ChaosRigidPhysicsAsync/RigidSceneAsync.h b/Engine/Plugins/Experimental/ChaosRigidPhysicsAsync/Source/ChaosRigidPhysicsAsync/Public/ChaosRigidPhysicsAsync/RigidSceneAsync.h
--- a/Engine/Plugins/Experimental/ChaosRigidPhysicsAsync/Source/ChaosRigidPhysicsAsync/Public/ChaosRigidPhysicsAsync/RigidSceneAsync.h
+++ b/Engine/Plugins/Experimental/ChaosRigidPhysicsAsync/Source/ChaosRigidPhysicsAsync/Public/ChaosRigidPhysicsAsync/RigidSceneAsync.h
@@ -11,6 +11,8 @@
 #include "ChaosRigidPhysicsAsync/RigidFactoryAsync.h"
 #include "ChaosRigidPhysicsAsync/RigidGeometryCollectionAsync.h"
 #include "ChaosRigidPhysicsAsync/RigidSceneSettingsAsync.h"
+#include "ChaosRigidPhysicsAsync/RigidShapeInstanceAsync.h"
+#include "ChaosRigidPhysicsAsync/ShapeInstanceContainerAsync.h"
 #include "ChaosSolversModule.h"
 #include "Containers/Map.h"
 #include "Logging/StructuredLog.h"
@@ -39,7 +41,7 @@ namespace Chaos::Rigids::Async
 	//
 	//
 
-	class UE_INTERNAL FRigidSceneAsyncGT 
+	class UE_INTERNAL FRigidSceneAsyncGT
 		: public UE::Physics::FRigidScene
 		, public TSharedFromThis<FRigidSceneAsyncGT, ESPMode::ThreadSafe>
 	{
@@ -65,7 +67,7 @@ namespace Chaos::Rigids::Async
 		// IRIGIDSCENE BEGIN
 		//
 	public:
-	
+
 		UE_INTERNAL virtual UE::Physics::FRigidSceneId GetId() const override final;
 		UE_INTERNAL void SetId(const UE::Physics::FRigidSceneId InId);
 		UE_INTERNAL virtual UE::Physics::FRigidSceneHandle GetHandle() const override final;
@@ -80,7 +82,16 @@ namespace Chaos::Rigids::Async
 		//
 		// IRIGIDSCENE END
 		//
-		
+
+		// ShapeInstance
+		UE_INTERNAL FRigidShapeInstanceAsync* AllocateShape();
+		UE_INTERNAL void DestroyShape(FRigidShapeInstanceAsync* InShape);
+		UE_INTERNAL virtual FRigidShapeInstanceAsync* GetShape(const FRigidObjectId& InId) const override final;
+		// ShapeInstance: Body
+		UE_INTERNAL void AddShape(FRigidBodyAsyncGT& Body, FRigidShapeInstanceAsync& InShapeInstance);
+		UE_INTERNAL void RemoveShape(FRigidBodyAsyncGT& Body, FRigidShapeInstanceAsync& InShapeInstance);
+		UE_INTERNAL FRigidShapeInstanceAsync* GetShape(const FShapeInstanceView& InView, const int32 InIndex) const;
+
 		// Joint Constraint
 		UE_INTERNAL virtual FJointConstraint6DOFAsyncGT* CreateJointConstraint6DOF() override final;
 		UE_INTERNAL virtual void DestroyJointConstraint(IJointConstraint* InConstraint) override final;
@@ -122,6 +133,7 @@ namespace Chaos::Rigids::Async
 		TArray<UE::Physics::IRigidSceneModifier*> SceneModifiers;
 
 		TUniquePtr<FRigidSceneAsyncPT> ScenePT;
+		FShapeInstanceContainerAsyncGT ShapeInstanceContainer;
 
 		// TODO_CHAOSAPI: Temp until Scene is always created through new API
 		bool bOwnsScene;
@@ -178,7 +190,14 @@ namespace Chaos::Rigids::Async
 		UE_INTERNAL CHAOSRIGIDPHYSICSASYNC_API void UnRegisterBody(const UE::Physics::FRigidBodyHandle& InBodyHandle, const UE::Physics::IRigidBody* InBody);
 		UE_INTERNAL CHAOSRIGIDPHYSICSASYNC_API void ActivateBody(const UE::Physics::FRigidBodyHandle& InBodyHandle);
 		UE_INTERNAL CHAOSRIGIDPHYSICSASYNC_API void DeactivateBody(const UE::Physics::FRigidBodyHandle& InBodyHandle);
-		
+
+		// ShapeInstance
+		UE_INTERNAL virtual FRigidShapeInstanceAsync* GetShape(const FRigidObjectId& InId) const override final;
+		// ShapeInstance: Body
+		UE_INTERNAL void AddShapePT(const FRigidBodyHandle BodyHandle, const FRigidShapeInstanceHandle ShapeInstanceHandle, const int32 ShapeIndex);
+		UE_INTERNAL void RemoveShapePT(const FRigidBodyHandle BodyHandle, const FRigidShapeInstanceHandle ShapeInstanceHandle);
+		UE_INTERNAL FRigidShapeInstanceAsync* GetShape(const FShapeInstanceView& InView, const int32 InIndex) const;
+
 		// Joint Constraint
 		UE_INTERNAL FJointConstraint6DOFAsyncPT* CreateJointConstraint6DOFPT(const FJointConstraint6DOFHandle& InHandle, Chaos::FJointConstraint* InJointConstraint);
 		UE_INTERNAL void DestroyJointConstraintPT(const FJointConstraint6DOFHandle& InHandle);
@@ -200,6 +219,7 @@ namespace Chaos::Rigids::Async
 		TMap<UE::Physics::FRigidObjectId, UE::Physics::IRigidBodyContainer*> BodyContainerRegistry;
 		TMap<UE::Physics::FRigidObjectId, UE::Physics::IRigidBody*> BodyRegistry;
 		TMap<UE::Physics::FRigidObjectId, UE::Physics::IJointConstraint*> JointConstraintRegistry;
+		FShapeInstanceContainerAsyncPT ShapeInstanceContainer;
 		TArray<UE::Physics::IRigidSceneModifier*> SceneModifiers;
 		bool bIsStarted;
 	};
diff --git a/Engine/Plugins/Experimental/ChaosRigidPhysicsAsync/Source/ChaosRigidPhysicsAsync/Public/ChaosRigidPhysicsAsync/RigidShapeInstanceAsync.h b/Engine/Plugins/Experimental/ChaosRigidPhysicsAsync/Source/ChaosRigidPhysicsAsync/Public/ChaosRigidPhysicsAsync/RigidShapeInstanceAsync.h
new file mode 100644
--- /dev/null
+++ b/Engine/Plugins/Experimental/ChaosRigidPhysicsAsync/Source/ChaosRigidPhysicsAsync/Public/ChaosRigidPhysicsAsync/RigidShapeInstanceAsync.h
@@ -0,0 +1,38 @@
+// Copyright Epic Games, Inc. All Rights Reserved.
+#pragma once 
+
+#include "ChaosRigidPhysicsAsync/RigidFwdAsync.h"
+
+#if UE_RIGIDPHYSICS_API_ENABLED
+
+#include "Chaos/CollisionFilterData.h"
+#include "Math/Transform.h"
+#include "RigidPhysics/RigidShapeInstance.h"
+
+namespace Chaos::Rigids::Async
+{
+	class UE_INTERNAL FRigidShapeInstanceAsync : public UE::Physics::IRigidShapeInstance
+	{
+		using FShapeFilterData = Chaos::Filter::FShapeFilterData;
+	public:
+		UE_RIGIDPHYSICS_RIGIDTYPED_DECL(CHAOSRIGIDPHYSICSASYNC_API, FRigidShapeInstanceAsync);
+
+		FRigidShapeInstanceAsync() = default;
+		virtual ~FRigidShapeInstanceAsync() = default;
+
+		UE_INTERNAL static Chaos::FPerShapeData* Create(const UE::Physics::FRigidShapeInstanceSetup& InSetup);
+
+		UE_INTERNAL virtual UE::Physics::FRigidShapeInstanceHandle GetHandle() const override final;
+		UE_INTERNAL virtual UE::Physics::FAnyGeometry GetGeometry() const override final;
+
+		UE_INTERNAL virtual FTransform3f GetLocalTransform() const override final;
+		UE_INTERNAL virtual FShapeFilterData GetShapeFilter() const override final;
+		UE_INTERNAL virtual void SetShapeFilter(const FShapeFilterData& ShapeFilter) override final;
+
+		UE_INTERNAL Chaos::FPerShapeData* PerShapeData = nullptr;
+		UE_INTERNAL UE::Physics::FRigidShapeInstanceHandle Handle;
+		UE_INTERNAL bool bAttachedToBody = false;
+	};
+} // namespace Chaos::Rigids::Async
+
+#endif // UE_RIGIDPHYSICS_API_ENABLED
diff --git a/Engine/Plugins/Experimental/ChaosRigidPhysicsAsync/Source/ChaosRigidPhysicsAsync/Public/ChaosRigidPhysicsAsync/ShapeInstanceContainerAsync.h b/Engine/Plugins/Experimental/ChaosRigidPhysicsAsync/Source/ChaosRigidPhysicsAsync/Public/ChaosRigidPhysicsAsync/ShapeInstanceContainerAsync.h
new file mode 100644
--- /dev/null
+++ b/Engine/Plugins/Experimental/ChaosRigidPhysicsAsync/Source/ChaosRigidPhysicsAsync/Public/ChaosRigidPhysicsAsync/ShapeInstanceContainerAsync.h
@@ -0,0 +1,79 @@
+// Copyright Epic Games, Inc. All Rights Reserved.
+#pragma once
+
+#include "ChaosRigidPhysicsAsync/RigidFwdAsync.h"
+
+#if UE_RIGIDPHYSICS_API_ENABLED
+
+#include "Chaos/ObjectPool.h"
+#include "ChaosRigidPhysicsAsync/RigidShapeInstanceAsync.h"
+#include "RigidPhysics/RigidObjectRegistry.h"
+#include "RigidPhysics/RigidSceneId.h"
+
+namespace Chaos::Rigids::Async
+{
+	class UE_INTERNAL FShapeInstanceView
+	{
+	public:
+		int32 Num() const;
+
+	private:
+		friend class FShapeInstanceContainerAsyncGT;
+		friend class FShapeInstanceContainerAsyncPT;
+		TArray<UE::Physics::FRigidObjectId> Ids;
+	};
+
+	class UE_INTERNAL FShapeInstanceContainerAsyncGT
+	{
+		using FRigidObjectId = UE::Physics::FRigidObjectId;
+		using FRigidSceneId = UE::Physics::FRigidSceneId;
+	public:
+		FShapeInstanceContainerAsyncGT();
+
+		FRigidShapeInstanceAsync* Allocate(const FRigidSceneId InSceneId);
+		void Free(FRigidShapeInstanceAsync* InInstance);
+
+		FRigidShapeInstanceAsync* Find(const FRigidObjectId& InId) const;
+
+		int32 Add(FRigidShapeInstanceAsync& InInstance, FShapeInstanceView& InOutView);
+		void Remove(FRigidShapeInstanceAsync& InInstance, FShapeInstanceView& InOutView);
+		FRigidShapeInstanceAsync* Get(const FShapeInstanceView& InView, const int32 InIndex) const;
+		void Free(FShapeInstanceView& InView);
+
+	private:
+		TObjectPool<FRigidShapeInstanceAsync> Pool;
+		UE::Physics::TRigidObjectRegistry<FRigidShapeInstanceAsync> Registry;
+	};
+
+	class UE_INTERNAL FShapeInstanceContainerAsyncPT
+	{
+		using FRigidObjectId = UE::Physics::FRigidObjectId;
+		using FRigidSceneId = UE::Physics::FRigidSceneId;
+		using FRigidShapeInstanceHandle = UE::Physics::FRigidShapeInstanceHandle;
+
+	public:
+		FShapeInstanceContainerAsyncPT();
+
+		FRigidShapeInstanceAsync* Allocate(const FRigidShapeInstanceHandle& InHandle);
+		void Free(FRigidShapeInstanceAsync* InInstance);
+
+		FRigidShapeInstanceAsync* Find(const FRigidObjectId& InId) const;
+
+		void Add(FRigidShapeInstanceAsync& InInstance, const int32 InShapeIndex, FShapeInstanceView& InOutView);
+		void Remove(FRigidShapeInstanceAsync& InInstance, FShapeInstanceView& InOutView);
+		FRigidShapeInstanceAsync* Get(const FShapeInstanceView& InView, const int32 InIndex) const;
+		void Free(FShapeInstanceView& InView);
+
+	private:
+		struct FRegistryItem
+		{
+			FRigidShapeInstanceAsync* Object = nullptr;
+			int32 Epoch = 0;
+		};
+
+		TObjectPool<FRigidShapeInstanceAsync> Pool;
+		TArray<FRegistryItem> Registry;
+	};
+} // namespace Chaos::Rigids::Async
+
+#endif // UE_RIGIDPHYSICS_API_ENABLED
diff --git a/Engine/Plugins/Experimental/ChaosRigidPhysicsAsync/Tests/Private/BodyTests.cpp b/Engine/Plugins/Experimental/ChaosRigidPhysicsAsync/Tests/Private/BodyTests.cpp
--- a/Engine/Plugins/Experimental/ChaosRigidPhysicsAsync/Tests/Private/BodyTests.cpp
+++ b/Engine/Plugins/Experimental/ChaosRigidPhysicsAsync/Tests/Private/BodyTests.cpp
@@ -13,16 +13,255 @@
 #include "RigidPhysics/RigidScene.h"
 #include "RigidPhysics/RigidShapeInstance.h"
 
+#include "RigidTestFixture.h"
 #include "RigidTestUtils.h"
 
 #if UE_RIGIDPHYSICS_API_ENABLED
 
 namespace Chaos::LowLevelTest
 {
-	using namespace Chaos;
 	using namespace Chaos::Rigids::Async;
 	using namespace UE::Physics;
 
+	TEST_CASE("RigidBody::ShapeInstance::CreateShape", "[Chaos][API][Body][ShapeInstance][unit]")
+	{
+		FRigidTestFixture Fixture;
+
+		FRigidBodyHandle BodyHandle;
+		const TArray<FRigidShapeInstanceSetup> ShapeSetups
+		{
+			FRigidShapeInstanceSetup(FBoxGeometry(FVector3f(1))),
+			FRigidShapeInstanceSetup(FSphereGeometry(2)),
+		};
+		TArray<FRigidShapeInstanceHandle> ShapeHandles;
+		ShapeHandles.SetNum(ShapeSetups.Num());
+
+		auto ValidateGTData = [&Fixture, &BodyHandle, &ShapeSetups, &ShapeHandles](int32 ShapeCount)
+			{
+				CAPTURE(ShapeCount);
+				if (FRigidContextGameRO Context = Fixture.LockROChecked())
+				{
+					TRigidBodyPtr<FRigidContextGameRO> BodyPtr = BodyHandle.Pin(Context);
+					REQUIRE(BodyPtr);
+					REQUIRE(BodyPtr->GetNumShapes() == ShapeCount);
+					for (int32 I = 0; I < ShapeCount; ++I)
+					{
+						const TRigidShapeInstancePtr<FRigidContextGameRO> ShapePtr = ShapeHandles[I].Pin(Context);
+						REQUIRE(ShapePtr);
+						CHECK(BodyPtr->GetShape(I) == ShapePtr);
+						CHECK(ShapePtr->GetGeometry() == ShapeSetups[I].Geometry);
+					}
+				}
+			};
+		auto ValidatePTData = [&Fixture, &BodyHandle, &ShapeSetups, &ShapeHandles](int32 ShapeCount)
+			{
+				CAPTURE(ShapeCount);
+				Fixture.RunPTCallback([&ShapeSetups, &ShapeHandles, BodyHandle, ShapeCount](const FRigidContextSimRW& Context)
+					{
+						TRigidBodyPtr<FRigidContextSimRW> BodyPtr = BodyHandle.Pin(Context);
+						REQUIRE(BodyPtr);
+						REQUIRE(BodyPtr->GetNumShapes() == ShapeCount);
+						for (int32 I = 0; I < ShapeCount; ++I)
+						{
+							const TRigidShapeInstancePtr<FRigidContextSimRW> ShapePtr = ShapeHandles[I].Pin(Context);
+							REQUIRE(ShapePtr);
+							CHECK(BodyPtr->GetShape(I) == ShapePtr);
+							CHECK(ShapePtr->GetGeometry() == ShapeSetups[I].Geometry);
+						}
+					});
+			};
+
+		SECTION("One Shot")
+		{
+			if (FRigidContextGameRW Context = Fixture.LockRWChecked())
+			{
+				TRigidBodyPtr<FRigidContextGameRW> BodyPtr = Context->CreateBody(FRigidDebugName("Body1"), ERigidMovementType::Dynamic);
+				BodyHandle = BodyPtr;
+				Fixture.AutoCleanup(BodyHandle);
+				for (int32 I = 0; I < ShapeSetups.Num(); ++I)
+				{
+					ShapeHandles[I] = BodyPtr->CreateShape(ShapeSetups[I]);
+				}
+				BodyPtr->Activate();
+			}
+
+			ValidateGTData(ShapeSetups.Num());
+			ValidatePTData(ShapeSetups.Num());
+		}
+		SECTION("Incremental")
+		{
+			if (FRigidContextGameRW Context = Fixture.LockRWChecked())
+			{
+				TRigidBodyPtr<FRigidContextGameRW> BodyPtr = Context->CreateBody(FRigidDebugName("Body1"), ERigidMovementType::Dynamic);
+				BodyHandle = BodyPtr;
+				Fixture.AutoCleanup(BodyHandle);
+				ShapeHandles[0] = BodyPtr->CreateShape(ShapeSetups[0]);
+				BodyPtr->Activate();
+			}
+			// Check initial config
+			ValidateGTData(1);
+			ValidatePTData(1);
+
+			// Incrementally add each shape
+			for (int32 Index = 1; Index < ShapeHandles.Num(); ++Index)
+			{
+				if (FRigidContextGameRW Context = Fixture.LockRWChecked())
+				{
+					const TRigidBodyPtr<FRigidContextGameRW> BodyPtr = BodyHandle.Pin(Context);
+					ShapeHandles[Index] = BodyPtr->CreateShape(ShapeSetups[Index]);
+				}
+
+				ValidateGTData(Index + 1);
+				ValidatePTData(Index + 1);
+			}
+		}
+	}
+
+	TEST_CASE("RigidBody::ShapeInstance::DestroyShape", "[Chaos][API][Body][ShapeInstance][unit]")
+	{
+		FRigidTestFixture Fixture;
+
+		FRigidBodyHandle BodyHandle;
+		const TArray<FRigidShapeInstanceSetup> ShapeSetups
+		{
+			FRigidShapeInstanceSetup(FBoxGeometry(FVector3f(1))),
+			FRigidShapeInstanceSetup(FSphereGeometry(2)),
+		};
+		TArray<FRigidShapeInstanceHandle> ShapeHandles;
+		ShapeHandles.SetNum(ShapeSetups.Num());
+
+		if (FRigidContextGameRW Context = Fixture.LockRWChecked())
+		{
+			TRigidBodyPtr<FRigidContextGameRW> BodyPtr = Context->CreateBody(FRigidDebugName("Body1"), ERigidMovementType::Dynamic);
+			BodyHandle = BodyPtr;
+			Fixture.AutoCleanup(BodyHandle);
+			for (int32 I = 0; I < ShapeSetups.Num(); ++I)
+			{
+				ShapeHandles[I] = BodyPtr->CreateShape(ShapeSetups[I]);
+			}
+			BodyPtr->Activate();
+		}
+
+		// Pump the thread to make sure the shapes have been propagated
+		Fixture.RunPTCallback([](const FRigidContextSimRW& Context) {});
+
+		SECTION("Remove One Shape")
+		{
+			const int32 RemovalIndex = GENERATE(0, 1);
+			const int32 RemainingIndex = (RemovalIndex + 1) % 2;
+			CAPTURE(RemovalIndex);
+			if (FRigidContextGameRW Context = Fixture.LockRWChecked())
+			{
+				TRigidBodyPtr<FRigidContextGameRW> BodyPtr = BodyHandle.Pin(Context);
+				BodyPtr->DestroyShape(ShapeHandles[RemovalIndex].Pin(Context));
+				REQUIRE(BodyPtr->GetNumShapes() == 1);
+				CHECK(BodyPtr->GetShape(0) == ShapeHandles[RemainingIndex].Pin(Context));
+				CHECK(!ShapeHandles[RemovalIndex].Pin(Context));
+				CHECK(ShapeHandles[RemainingIndex].Pin(Context));
+			}
+			if (FRigidContextGameRO Context = Fixture.LockROChecked())
+			{
+				const TRigidBodyPtr<FRigidContextGameRO> BodyPtr = BodyHandle.Pin(Context);
+				REQUIRE(BodyPtr->GetNumShapes() == 1);
+				CHECK(BodyPtr->GetShape(0) == ShapeHandles[RemainingIndex].Pin(Context));
+				CHECK(!ShapeHandles[RemovalIndex].Pin(Context));
+				CHECK(ShapeHandles[RemainingIndex].Pin(Context));
+			}
+			Fixture.RunPTCallback([&ShapeSetups, &ShapeHandles, BodyHandle, RemovalIndex, RemainingIndex](const FRigidContextSimRW& Context)
+				{
+					const TRigidBodyPtr<FRigidContextSimRW> BodyPtr = BodyHandle.Pin(Context);
+					REQUIRE(BodyPtr->GetNumShapes() == 1);
+					CHECK(BodyPtr->GetShape(0) == ShapeHandles[RemainingIndex].Pin(Context));
+					CHECK(!ShapeHandles[RemovalIndex].Pin(Context));
+					CHECK(ShapeHandles[RemainingIndex].Pin(Context));
+				});
+		}
+		SECTION("Remove All Shapes")
+		{
+			if (FRigidContextGameRW Context = Fixture.LockRWChecked())
+			{
+				TRigidBodyPtr<FRigidContextGameRW> BodyPtr = BodyHandle.Pin(Context);
+				for (int32 I = 0; I < ShapeHandles.Num(); ++I)
+				{
+					BodyPtr->DestroyShape(ShapeHandles[I].Pin(Context));
+					CHECK(!ShapeHandles[I].Pin(Context));
+				}
+				CHECK(BodyPtr->GetNumShapes() == 0);
+			}
+			if (FRigidContextGameRO Context = Fixture.LockROChecked())
+			{
+				CHECK(BodyHandle.Pin(Context)->GetNumShapes() == 0);
+				for (int32 I = 0; I < ShapeHandles.Num(); ++I)
+				{
+					CHECK(!ShapeHandles[I].Pin(Context));
+				}
+			}
+			Fixture.RunPTCallback([&ShapeHandles, BodyHandle](const FRigidContextSimRW& Context)
+				{
+					CHECK(BodyHandle.Pin(Context)->GetNumShapes() == 0);
+					for (int32 I = 0; I < ShapeHandles.Num(); ++I)
+					{
+						CHECK(!ShapeHandles[I].Pin(Context));
+					}
+				});
+		}
+	}
+
+	TEST_CASE("RigidBody::ShapeInstance::Destroy Body With Shapes", "[Chaos][API][Body][ShapeInstance][unit]")
+	{
+		FRigidTestFixture Fixture;
+
+		FRigidBodyHandle BodyHandle;
+		const TArray<FRigidShapeInstanceSetup> ShapeSetups
+		{
+			FRigidShapeInstanceSetup(FBoxGeometry(FVector3f(1))),
+			FRigidShapeInstanceSetup(FSphereGeometry(2)),
+		};
+		TArray<FRigidShapeInstanceHandle> ShapeHandles;
+		ShapeHandles.SetNum(ShapeSetups.Num());
+
+		if (FRigidContextGameRW Context = Fixture.LockRWChecked())
+		{
+			TRigidBodyPtr<FRigidContextGameRW> BodyPtr = Context->CreateBody(FRigidDebugName("Body1"), ERigidMovementType::Dynamic);
+			BodyHandle = BodyPtr;
+			for (int32 I = 0; I < ShapeSetups.Num(); ++I)
+			{
+				ShapeHandles[I] = BodyPtr->CreateShape(ShapeSetups[I]);
+			}
+			BodyPtr->Activate();
+		}
+
+		// Pump the thread to make sure the shapes have been propagated
+		Fixture.RunPTCallback([](const FRigidContextSimRW& Context) {});
+
+		if (FRigidContextGameRW Context = Fixture.LockRWChecked())
+		{
+			Context->DestroyBody(BodyHandle.Pin(Context));
+
+			CHECK(!BodyHandle.Pin(Context).IsValid());
+			for (int32 I = 0; I < ShapeHandles.Num(); ++I)
+			{
+				CHECK(!ShapeHandles[I].Pin(Context).IsValid());
+			}
+		}
+		if (FRigidContextGameRO Context = Fixture.LockROChecked())
+		{
+			CHECK(!BodyHandle.Pin(Context).IsValid());
+			for (int32 I = 0; I < ShapeHandles.Num(); ++I)
+			{
+				CHECK(!ShapeHandles[I].Pin(Context).IsValid());
+			}
+		}
+		Fixture.RunPTCallback([BodyHandle, &ShapeHandles](const FRigidContextSimRW& Context)
+			{
+				CHECK(!BodyHandle.Pin(Context).IsValid());
+				for (int32 I = 0; I < ShapeHandles.Num(); ++I)
+				{
+					CHECK(!ShapeHandles[I].Pin(Context).IsValid());
+				}
+			});
+	}
+
 	// Create bodies with various setups and validate them
 	TEST_CASE("BodySettingsTest", "[Chaos][API][Body][unit]")
 	{
diff --git a/Engine/Plugins/Experimental/ChaosRigidPhysicsAsync/Tests/Private/RigidTestFixture.cpp b/Engine/Plugins/Experimental/ChaosRigidPhysicsAsync/Tests/Private/RigidTestFixture.cpp
--- a/Engine/Plugins/Experimental/ChaosRigidPhysicsAsync/Tests/Private/RigidTestFixture.cpp
+++ b/Engine/Plugins/Experimental/ChaosRigidPhysicsAsync/Tests/Private/RigidTestFixture.cpp
@@ -22,7 +22,7 @@ namespace Chaos::LowLevelTest
 	using namespace Chaos;
 	using namespace Chaos::Rigids::Async;
 	using namespace UE::Physics;
-	
+
 	FRigidTestFixture::FRigidTestFixture()
 	{
 		SceneHandle = Factory.CreateScene(FRigidDebugName::Make(TEXT("Main")), &SceneSettings);
@@ -48,6 +48,20 @@ namespace Chaos::LowLevelTest
 		Factory.DestroyScene(SceneHandle);
 	}
 
+	FRigidContextGameRO FRigidTestFixture::LockROChecked() const
+	{
+		// Note: This isn't 100% identical to checking the context, but it's effectively the same and there's no other way to write this.
+		CHECK(SceneHandle.GetScene());
+		return SceneHandle.LockRO();
+	}
+
+	FRigidContextGameRW FRigidTestFixture::LockRWChecked() const
+	{
+		// Note: This isn't 100% identical to checking the context, but it's effectively the same and there's no other way to write this.
+		CHECK(SceneHandle.GetScene());
+		return SceneHandle.LockRW();
+	}
+
 	void FRigidTestFixture::AutoCleanup(FJointConstraintHandle Handle)
 	{
 		JointHandles.Add(Handle);
@@ -80,4 +94,4 @@ namespace Chaos::LowLevelTest
 	}
 } // namespace Chaos::LowLevelTest
 
-#endif
+#endif UE_RIGIDPHYSICS_API_ENABLED
diff --git a/Engine/Plugins/Experimental/ChaosRigidPhysicsAsync/Tests/Private/RigidTestFixture.h b/Engine/Plugins/Experimental/ChaosRigidPhysicsAsync/Tests/Private/RigidTestFixture.h
--- a/Engine/Plugins/Experimental/ChaosRigidPhysicsAsync/Tests/Private/RigidTestFixture.h
+++ b/Engine/Plugins/Experimental/ChaosRigidPhysicsAsync/Tests/Private/RigidTestFixture.h
@@ -6,8 +6,10 @@
 #include "ChaosRigidPhysicsAsync/RigidSceneSettingsAsync.h"
 #include "RigidPhysics/RigidScene.h"
 
+#if UE_RIGIDPHYSICS_API_ENABLED
+
 namespace Chaos::LowLevelTest
-{	
+{
 	using FRigidDebugName = UE::Physics::FRigidDebugName;
 	using FRigidContextGameRW = UE::Physics::FRigidContextGameRW;
 	using FRigidContextSimRW = UE::Physics::FRigidContextSimRW;
@@ -15,12 +17,17 @@ namespace Chaos::LowLevelTest
 	using FJointConstraintHandle = UE::Physics::FJointConstraintHandle;
 	using FRigidBodyHandle = UE::Physics::FRigidBodyHandle;
 	using FRigidGeometryCollectionHandle = UE::Physics::FRigidGeometryCollectionHandle;
+	using FRigidContextGameRO = UE::Physics::FRigidContextGameRO;
+	using FRigidContextGameRW = UE::Physics::FRigidContextGameRW;
 
 	struct FRigidTestFixture
 	{
 		FRigidTestFixture();
 		~FRigidTestFixture();
 
+		FRigidContextGameRO LockROChecked() const;
+		FRigidContextGameRW LockRWChecked() const;
+
 		void AutoCleanup(FJointConstraintHandle Handle);
 		void AutoCleanup(FRigidBodyHandle Handle);
 		void AutoCleanup(FRigidGeometryCollectionHandle Handle);
@@ -35,4 +42,6 @@ namespace Chaos::LowLevelTest
 		TArray<FRigidGeometryCollectionHandle> GeometryCollectionHandles;
 		TArray<FJointConstraintHandle> JointHandles;
 	};
-} // namespace Chaos::LowLevelTest
\ No newline at end of file
+} // namespace Chaos::LowLevelTest
+
+#endif // UE_RIGIDPHYSICS_API_ENABLED
diff --git a/Engine/Plugins/Experimental/ChaosRigidPhysicsAsync/Tests/Private/RigidTestUtils.cpp b/Engine/Plugins/Experimental/ChaosRigidPhysicsAsync/Tests/Private/RigidTestUtils.cpp
--- a/Engine/Plugins/Experimental/ChaosRigidPhysicsAsync/Tests/Private/RigidTestUtils.cpp
+++ b/Engine/Plugins/Experimental/ChaosRigidPhysicsAsync/Tests/Private/RigidTestUtils.cpp
@@ -4,14 +4,19 @@
 
 #if UE_RIGIDPHYSICS_API_ENABLED
 
-#include "Chaos/CollisionFilterData.h"
 #include "PhysicalMaterials/PhysicalMaterial.h"
-#include "Physics/Experimental/ChaosInterfaceUtils.h"
-#include "Physics/PhysicsInterfaceTypes.h"
-#include "PhysicsEngine/AggregateGeom.h"
 
 namespace UE::Physics
 {
+	FRigidShapeInstanceSetup MakeBoxShape(const FVector3f& InSize, const FVector3f& InCenter, UPhysicalMaterial* InMaterial)
+	{
+		FBoxGeometry BoxGeom(InSize);
+		FRigidShapeInstanceSetup Setup(BoxGeom);
+		Setup.LocalTransform = FTransform3f(InCenter);
+		// TODO: Material
+		return Setup;
+	}
+
 	FVector3f MakeSolidBoxInertia(float InMass, const FVector3f& InSize)
 	{
 		FVector3f SizeSq = InSize * InSize;
@@ -20,85 +25,6 @@ namespace UE::Physics
 			SizeSq.X + SizeSq.Z,
 			SizeSq.X + SizeSq.Y);
 	}
+} // namespace UE::Physics
 
-	FRigidShapeInstancePtr MakeShapeInstanceFromAddParams(const FGeometryAddParams& InParams)
-	{
-		TArray<Chaos::FImplicitObjectPtr> Geoms;
-		Chaos::FShapesArray Shapes;
-		ChaosInterface::CreateGeometry(InParams, Geoms, Shapes);
-
-		if (Shapes.Num() != 1)
-		{
-			UE_LOGFMT(LogRigidPhysics, Warning, "MakeShapeInstanceFromAddParams is intended to be used for single shapes only. Have {0} shapes.", Shapes.Num());
-		}
-
-		if (Shapes.Num() > 0)
-		{
-			return FRigidShapeInstancePtr(Shapes[0].Release());
-		}
-
-		return FRigidShapeInstancePtr();
-	}
-
-	FRigidShapeInstancePtr MakeColliderFromAggGeom(
-		FKAggregateGeom& InAggGeom, 
-		const FVector3d& InScale,
-		ECollisionTraceFlag CollisionTraceFlag,
-		UPhysicalMaterial* InSimpleMaterial,
-		const TArrayView<UPhysicalMaterial*>& InComplexMaterials,
-		const TArrayView<FPhysicalMaterialMaskParams>& InComplexMaterialMasks,
-		const FTransform& InRelativeTransform)
-	{
-		// Not used - filters are set up after
-		FBodyCollisionData BodyCollisionData;
-
-		FGeometryAddParams AddParams;
-		AddParams.bDoubleSided = false;
-		AddParams.CollisionData = BodyCollisionData;
-		AddParams.CollisionTraceType = CollisionTraceFlag;
-		AddParams.Scale = InScale;
-		AddParams.SimpleMaterial = InSimpleMaterial;
-		AddParams.ComplexMaterials = InComplexMaterials;
-		AddParams.ComplexMaterialMasks = InComplexMaterialMasks;
-		AddParams.LocalTransform = InRelativeTransform;
-		AddParams.WorldTransform = FTransform::Identity;
-		AddParams.Geometry = &InAggGeom;
-		//AddParams.TriMeshGeometries = MakeArrayView(TriMeshGeometries);
-
-		return MakeShapeInstanceFromAddParams(AddParams);
-	}
-
-	FRigidShapeInstancePtr MakeBoxShape(const FVector3f& InSize, const FVector3f& InCenter, UPhysicalMaterial* InMaterial)
-	{
-		using namespace Chaos;
-
-		// Single Box
-		FKBoxElem BoxElem = FKBoxElem(InSize.X, InSize.Y, InSize.Z);
-		BoxElem.Center = FVector(InCenter);
-
-		FKAggregateGeom AggGeom;
-		AggGeom.AddElement(BoxElem);
-
-		FRigidShapeInstancePtr Collider = MakeColliderFromAggGeom(
-			AggGeom, 
-			FVector3d::One(),
-			ECollisionTraceFlag::CTF_UseSimpleAsComplex,
-			InMaterial,
-			{},
-			{},
-			FTransform3d::Identity);
-
-		if (Collider.IsValid())
-		{
-			const EFilterFlags FilterFlags = EFilterFlags::SimpleCollision | EFilterFlags::ContactNotify | EFilterFlags::ModifyContacts;
-			const Filter::FShapeFilterData ShapeFilter = Filter::FShapeFilterBuilder::BuildBlockAll(FilterFlags);
-			Collider->GetPerShapeData()->SetQueryEnabled(true);
-			Collider->GetPerShapeData()->SetSimEnabled(true);
-			Collider->GetPerShapeData()->SetShapeFilterData(ShapeFilter);
-		}
-
-		return Collider;
-	}
-}
-
-#endif
\ No newline at end of file
+#endif // UE_RIGIDPHYSICS_API_ENABLED
diff --git a/Engine/Plugins/Experimental/ChaosRigidPhysicsAsync/Tests/Private/RigidTestUtils.h b/Engine/Plugins/Experimental/ChaosRigidPhysicsAsync/Tests/Private/RigidTestUtils.h
--- a/Engine/Plugins/Experimental/ChaosRigidPhysicsAsync/Tests/Private/RigidTestUtils.h
+++ b/Engine/Plugins/Experimental/ChaosRigidPhysicsAsync/Tests/Private/RigidTestUtils.h
@@ -8,21 +8,19 @@
 #include "Math/Bounds.h"
 #include "Math/Transform.h"
 #include "RigidPhysics/RigidLog.h"
-#include "RigidPhysics/RigidShapeInstance.h"
+#include "RigidPhysics/RigidShapeInstanceSetup.h"
 
 class UPhysicalMaterial;
 
 namespace Chaos
 {
 	class FPerShapeData;
-}
+} // namespace Chaos
 
 namespace UE::Physics
 {
-	// Utility for Unit Tests until the API if more fully implemented
-	FRigidShapeInstancePtr MakeBoxShape(const FVector3f& InSize, const FVector3f& InCenter = FVector3f::Zero(), UPhysicalMaterial* InMaterial = nullptr);
+	FRigidShapeInstanceSetup MakeBoxShape(const FVector3f& InSize, const FVector3f& InCenter = FVector3f::Zero(), UPhysicalMaterial* InMaterial = nullptr);
 	FVector3f MakeSolidBoxInertia(float InMass, const FVector3f& InSize);
+} // namespace UE::Physics
 
-}
-
-#endif
\ No newline at end of file
+#endif // UE_RIGIDPHYSICS_API_ENABLED
diff --git a/Engine/Plugins/Experimental/ChaosRigidPhysicsAsync/Tests/Private/ShapeInstanceTests.cpp b/Engine/Plugins/Experimental/ChaosRigidPhysicsAsync/Tests/Private/ShapeInstanceTests.cpp
new file mode 100644
--- /dev/null
+++ b/Engine/Plugins/Experimental/ChaosRigidPhysicsAsync/Tests/Private/ShapeInstanceTests.cpp
@@ -0,0 +1,165 @@
+// Copyright Epic Games, Inc. All Rights Reserved.
+
+#include "Chaos/LowLevelTest/ChaosTestHarness.h"
+#include "Chaos/LowLevelTest/ChaosTestScene.h"
+#include "ChaosRigidPhysicsAsyncTest.h"
+#include "RigidPhysics/RigidBody.h"
+#include "RigidPhysics/RigidContext.h"
+#include "RigidPhysics/RigidScene.h"
+#include "RigidPhysics/RigidShapeInstance.h"
+#include "RigidTestFixture.h"
+
+#if UE_RIGIDPHYSICS_API_ENABLED
+
+namespace Chaos::LowLevelTest
+{
+	using namespace UE::Physics;
+
+	template <typename GeometryType>
+	void CheckGeometryTransform(const FTransform3f& Expected, const FTransform3f& Actual)
+	{
+		// TODO: The Catch2 StringMaker isn't working correctly for FVector3f for some reason. Manually casting to Chaos::FVec3f for now...
+		CHECK_THAT(Chaos::FVec3f(Expected.GetLocation()), Catch::ApproxEq(Chaos::FVec3f(Actual.GetLocation())));
+		CHECK_THAT(Expected.GetRotation(), Catch::ApproxEq(Actual.GetRotation()));
+		CHECK_THAT(Chaos::FVec3f(Expected.GetScale3D()), Catch::ApproxEq(Chaos::FVec3f(Actual.GetScale3D())));
+	}
+
+	void CreateAndCheckShapeInstanceAndTransform(FRigidTestFixture& Fixture, const FRigidShapeInstanceSetup& InSetup, const FAnyGeometry& ExpectedGeometry, const FTransform3f& ExpectedTransform)
+	{
+		if (FRigidContextGameRW Context = Fixture.LockRWChecked())
+		{
+			TRigidBodyPtr<FRigidContextGameRW> BodyPtr = Context->CreateBody(FRigidDebugName("Body1"), ERigidMovementType::Dynamic);
+			Fixture.AutoCleanup(BodyPtr);
+
+			TRigidShapeInstancePtr<FRigidContextGameRW> ShapeInstance = BodyPtr->CreateShape(InSetup);
+
+			const FTransform3f ActualTransform = ShapeInstance->GetLocalTransform();
+			const FAnyGeometry ActualGeometry = ShapeInstance->GetGeometry();
+			CHECK(ExpectedGeometry == ActualGeometry);
+			// Visit the actual geometry type so we can control how we check the transform.
+			// Some types don't have a unique transform, so instead we have to test the "pose" of the transform.
+			ActualGeometry.Visit([&ExpectedTransform, &ActualTransform]<typename GeometryType>(const GeometryType& ActualGeometry)
+			{
+				CheckGeometryTransform<GeometryType>(ExpectedTransform, ActualTransform);
+			});
+		}
+	}
+
+	void CreateAndCheckShapeInstanceAndTransform(FRigidTestFixture& Fixture, const FRigidShapeInstanceSetup& InSetup)
+	{
+		CreateAndCheckShapeInstanceAndTransform(Fixture, InSetup, InSetup.Geometry, InSetup.LocalTransform);
+	}
+
+	TEST_CASE("ShapeInstance: Setup + GetShape/Transform", "[Chaos][API][ShapeInstance][unit]")
+	{
+		FRigidTestFixture Fixture;
+		SECTION("Box")
+		{
+			FRigidShapeInstanceSetup Setup(FBoxGeometry(FVector3f(1, 2, 3)));
+
+			SECTION("Identity Transform")
+			{
+				Setup.LocalTransform.SetIdentity();
+				CreateAndCheckShapeInstanceAndTransform(Fixture, Setup);
+			}
+			SECTION("Translation")
+			{
+				Setup.LocalTransform.SetTranslation(FVector3f(0, 1, 0));
+				CreateAndCheckShapeInstanceAndTransform(Fixture, Setup);
+			}
+			SECTION("Rotation")
+			{
+				Setup.LocalTransform.SetRotation(FRotator3f(10, 20, 30).Quaternion());
+				CreateAndCheckShapeInstanceAndTransform(Fixture, Setup);
+			}
+			SECTION("Scale")
+			{
+				Setup.LocalTransform.SetScale3D(FVector3f(2, 3, 4));
+				CreateAndCheckShapeInstanceAndTransform(Fixture, Setup, FBoxGeometry(FVector3f(2, 6, 12)), FTransform3f());
+			}
+			SECTION("Full Transform")
+			{
+				Setup.LocalTransform.SetLocation(FVector3f(-1, -2, -3));
+				Setup.LocalTransform.SetRotation(FRotator3f(10, 20, 30).Quaternion());
+				Setup.LocalTransform.SetScale3D(FVector3f(2, 3, 4));
+				FTransform3f ExpectedTransform = Setup.LocalTransform;
+				ExpectedTransform.SetScale3D(FVec3f::One());
+				CreateAndCheckShapeInstanceAndTransform(Fixture, Setup, FBoxGeometry(FVector3f(2, 6, 12)), ExpectedTransform);
+			}
+		}
+		SECTION("Sphere")
+		{
+			FRigidShapeInstanceSetup Setup(FSphereGeometry(2));
+
+			SECTION("Identity Transform")
+			{
+				Setup.LocalTransform.SetIdentity();
+				CreateAndCheckShapeInstanceAndTransform(Fixture, Setup);
+			}
+			SECTION("Translation")
+			{
+				Setup.LocalTransform.SetTranslation(FVector3f(0, 1, 0));
+				CreateAndCheckShapeInstanceAndTransform(Fixture, Setup);
+			}
+			SECTION("Rotation")
+			{
+				Setup.LocalTransform.SetRotation(FRotator3f(10, 20, 30).Quaternion());
+				CreateAndCheckShapeInstanceAndTransform(Fixture, Setup, Setup.Geometry, FTransform3f());
+			}
+			SECTION("Uniform Scale")
+			{
+				Setup.LocalTransform.SetScale3D(FVector3f(3));
+				CreateAndCheckShapeInstanceAndTransform(Fixture, Setup, FSphereGeometry(6), FTransform3f());
+			}
+			SECTION("Non-Uniform Scale")
+			{
+				// The min scale element is always used
+				const FVector3f Scale = GENERATE(FVector3f(2, 3, 4), FVector3f(3, 4, 2), FVector3f(4, 2, 3));
+				Setup.LocalTransform.SetScale3D(Scale);
+				CreateAndCheckShapeInstanceAndTransform(Fixture, Setup, FSphereGeometry(4), FTransform3f());
+			}
+			SECTION("Full Transform")
+			{
+				Setup.LocalTransform.SetLocation(FVector3f(-1, -2, -3));
+				Setup.LocalTransform.SetRotation(FRotator3f(10, 20, 30).Quaternion());
+				Setup.LocalTransform.SetScale3D(FVector3f(3));
+				FTransform3f ExpectedTransform;
+				ExpectedTransform.SetLocation(Setup.LocalTransform.GetLocation());
+				CreateAndCheckShapeInstanceAndTransform(Fixture, Setup, FSphereGeometry(6), ExpectedTransform);
+			}
+		}
+	}
+
+	TEST_CASE("ShapeInstance: ShapeFilter", "[Chaos][API][ShapeInstance][unit]")
+	{
+		FRigidTestFixture Fixture;
+		FRigidShapeInstanceSetup Setup(FBoxGeometry(FVector3f(1, 2, 3)));
+
+		const Chaos::Filter::FShapeFilterData ShapeFilter0 = Chaos::Filter::FShapeFilterBuilder().SetCollisionChannelIndex(1).Build();
+		const Chaos::Filter::FShapeFilterData ShapeFilter1 = Chaos::Filter::FShapeFilterBuilder().SetCollisionChannelIndex(2).Build();
+
+		FRigidShapeInstanceHandle ShapeInstanceHandle;
+		if (FRigidContextGameRW Context = Fixture.LockRWChecked())
+		{
+			TRigidBodyPtr<FRigidContextGameRW> BodyPtr = Context->CreateBody(FRigidDebugName("Body1"), ERigidMovementType::Dynamic);
+			Fixture.AutoCleanup(BodyPtr);
+
+			TRigidShapeInstancePtr<FRigidContextGameRW> ShapeInstance0 = BodyPtr->CreateShape(Setup);
+			ShapeInstanceHandle = ShapeInstance0;
+			BodyPtr->Activate();
+
+			ShapeInstance0->SetShapeFilter(ShapeFilter0);
+			CHECK(ShapeInstance0->GetShapeFilter() == ShapeFilter0);
+			ShapeInstance0->SetShapeFilter(ShapeFilter1);
+			CHECK(ShapeInstance0->GetShapeFilter() == ShapeFilter1);
+		}
+
+		Fixture.RunPTCallback([ShapeInstanceHandle, &ShapeFilter1](const FRigidContextSimRW& Context)
+			{
+				const TRigidShapeInstancePtr<FRigidContextSimRW> ShapeInstancePtr = ShapeInstanceHandle.Pin(Context);
+				CHECK(ShapeInstancePtr->GetShapeFilter() == ShapeFilter1);
+			});
+	}
+} // namespace Chaos::LowLevelTest
+
+#endif // UE_RIGIDPHYSICS_API_ENABLED
diff --git a/Engine/Source/Developer/ChaosTestHarness/Public/Chaos/LowLevelTest/ChaosTestMatchers.h b/Engine/Source/Developer/ChaosTestHarness/Public/Chaos/LowLevelTest/ChaosTestMatchers.h
--- a/Engine/Source/Developer/ChaosTestHarness/Public/Chaos/LowLevelTest/ChaosTestMatchers.h
+++ b/Engine/Source/Developer/ChaosTestHarness/Public/Chaos/LowLevelTest/ChaosTestMatchers.h
@@ -43,7 +43,51 @@ namespace Catch
 	{
 		return FVectorEqualityMatcher<T, D>(Target, Epsilon);
 	}
-	
+
+	template <typename T>
+	FVectorEqualityMatcher<T, 3> ApproxEq(const UE::Math::TVector<T>& Target, const T Epsilon = UE_KINDA_SMALL_NUMBER)
+	{
+		return ApproxEq(Chaos::TVec3<T>(Target), Epsilon);
+	}
+
+	template <typename T>
+	class FQuatEqualityMatcher : public Catch::Matchers::MatcherBase<UE::Math::TQuat<T>>
+	{
+	public:
+		using FQuat = UE::Math::TQuat<T>;
+		FQuat Target;
+		T Epsilon;
+
+		FQuatEqualityMatcher(const FQuat& InTarget, const T InEpsilon = UE_KINDA_SMALL_NUMBER)
+			: Target(InTarget)
+			, Epsilon(InEpsilon)
+		{
+		}
+
+		// Performs the test for this matcher
+		bool match(const FQuat& ToCompare) const override
+		{
+			return Target.Equals(ToCompare, Epsilon);
+		}
+
+		// Produces a string describing what this matcher does. It should
+		// include any provided data (the begin/ end in this case) and
+		// be written as if it were stating a fact (in the output it will be
+		// preceded by the value under test).
+		virtual std::string describe() const override
+		{
+			std::ostringstream Stream;
+			Stream << "is equal to (" << Catch::StringMaker<FQuat>::convert(Target) << ") with epsilon " << Epsilon;
+			return Stream.str();
+		}
+	};
+
+	template <typename T>
+	FQuatEqualityMatcher<T> ApproxEq(const UE::Math::TQuat<T>& Target, const T Epsilon = UE_KINDA_SMALL_NUMBER)
+	{
+		return FQuatEqualityMatcher<T>(Target, Epsilon);
+	}
+
 	template <typename T, int D>
 	class TAABBEqualityMatcher : public Catch::Matchers::MatcherBase<Chaos::TAABB<T, D>>
 	{
diff --git a/Engine/Source/Developer/ChaosTestHarness/Public/Chaos/LowLevelTest/ChaosTestStringMakers.h b/Engine/Source/Developer/ChaosTestHarness/Public/Chaos/LowLevelTest/ChaosTestStringMakers.h
--- a/Engine/Source/Developer/ChaosTestHarness/Public/Chaos/LowLevelTest/ChaosTestStringMakers.h
+++ b/Engine/Source/Developer/ChaosTestHarness/Public/Chaos/LowLevelTest/ChaosTestStringMakers.h
@@ -8,8 +8,14 @@
 
 #include <string>
 
-namespace Catch 
+namespace Catch
 {
+	template<typename T>
+	struct StringMaker<UE::Math::TQuat<T>>
+	{
+		static std::string convert(UE::Math::TQuat<T> const& Value) { return StringMaker<FString>::convert(Value.ToString()); }
+	};
+
 	template<>
 	struct StringMaker<Chaos::FRotation3>
 	{
@@ -34,16 +40,10 @@ namespace Catch
 		static std::string convert(Chaos::FVec3f const& Value) { return std::string(TCHAR_TO_UTF8(*Chaos::ToString(Value))); }
 	};
 
-	template<>
-	struct StringMaker<FVector3d>
-	{
-		static std::string convert(FVector3d const& Value) { return StringMaker<Chaos::FVec3>::convert(Chaos::FVec3(Value)); }
-	};
-
-	template<>
-	struct StringMaker<FVector3f>
+	template<typename T>
+	struct StringMaker<UE::Math::TVector<T>>
 	{
-		static std::string convert(FVector3f const& Value) { return StringMaker<Chaos::FVec3f>::convert(Chaos::FVec3f(Value)); }
+		static std::string convert(UE::Math::TVector<T> const& Value) { return StringMaker<Chaos::TVec3<T>>::convert(Chaos::TVec3<T>(Value)); }
 	};
 
 	template<typename T, int d>
diff --git a/Engine/Source/Runtime/Experimental/RigidPhysics/Private/RigidPhysics/Internal/IRigidBody.cpp b/Engine/Source/Runtime/Experimental/RigidPhysics/Private/RigidPhysics/Internal/IRigidBody.cpp
--- a/Engine/Source/Runtime/Experimental/RigidPhysics/Private/RigidPhysics/Internal/IRigidBody.cpp
+++ b/Engine/Source/Runtime/Experimental/RigidPhysics/Private/RigidPhysics/Internal/IRigidBody.cpp
@@ -41,7 +41,19 @@ namespace UE::Physics
 		return false;
 	}
 
-	void IRigidBody::AddShape(const FConstRigidShapeInstancePtr& InShape)
+	IRigidShapeInstance* IRigidBody::CreateShape(const FRigidShapeInstanceSetup& InSetup)
+	{
+		RIGIDPHYSICS_API_UNSUPPORTED();
+		return nullptr;
+	}
+
+	IRigidShapeInstance* IRigidBody::CreateShape(TUniquePtr<Chaos::FPerShapeData>&& InShapeData)
+	{
+		RIGIDPHYSICS_API_UNSUPPORTED();
+		return nullptr;
+	}
+
+	void IRigidBody::DestroyShape(IRigidShapeInstance* InShape, bool bWakeTouching)
 	{
 		RIGIDPHYSICS_API_UNSUPPORTED();
 	}
diff --git a/Engine/Source/Runtime/Experimental/RigidPhysics/Private/RigidPhysics/Internal/IRigidShapeInstance.cpp b/Engine/Source/Runtime/Experimental/RigidPhysics/Private/RigidPhysics/Internal/IRigidShapeInstance.cpp
new file mode 100644
--- /dev/null
+++ b/Engine/Source/Runtime/Experimental/RigidPhysics/Private/RigidPhysics/Internal/IRigidShapeInstance.cpp
@@ -0,0 +1,13 @@
+// Copyright Epic Games, Inc. All Rights Reserved.
+#include "RigidPhysics/Internal/IRigidShapeInstance.h"
+
+#include "RigidPhysics/RigidLog.h"
+
+#if UE_RIGIDPHYSICS_API_ENABLED
+
+namespace UE::Physics
+{
+	UE_RIGIDPHYSICS_RIGIDTYPED_IMPL(IRigidShapeInstance, IRigidTyped);
+} // namespace UE::Physics
+
+#endif // UE_RIGIDPHYSICS_API_ENABLED
diff --git a/Engine/Source/Runtime/Experimental/RigidPhysics/Private/RigidPhysics/RigidBodyHandle.cpp b/Engine/Source/Runtime/Experimental/RigidPhysics/Private/RigidPhysics/RigidBodyHandle.cpp
--- a/Engine/Source/Runtime/Experimental/RigidPhysics/Private/RigidPhysics/RigidBodyHandle.cpp
+++ b/Engine/Source/Runtime/Experimental/RigidPhysics/Private/RigidPhysics/RigidBodyHandle.cpp
@@ -18,6 +18,11 @@ namespace UE::Physics
 		return BodyId;
 	}
 
+	const FRigidSceneId& FRigidBodyHandle::GetSceneId() const
+	{
+		return SceneId;
+	}
+
 	void FRigidBodyHandle::Reset()
 	{
 		SceneId.Reset();
diff --git a/Engine/Source/Runtime/Experimental/RigidPhysics/Private/RigidPhysics/RigidShapeInstance.cpp b/Engine/Source/Runtime/Experimental/RigidPhysics/Private/RigidPhysics/RigidShapeInstance.cpp
deleted file mode 100644
diff --git a/Engine/Source/Runtime/Experimental/RigidPhysics/Private/RigidPhysics/RigidShapeInstanceHandle.cpp b/Engine/Source/Runtime/Experimental/RigidPhysics/Private/RigidPhysics/RigidShapeInstanceHandle.cpp
new file mode 100644
--- /dev/null
+++ b/Engine/Source/Runtime/Experimental/RigidPhysics/Private/RigidPhysics/RigidShapeInstanceHandle.cpp
@@ -0,0 +1,31 @@
+// Copyright Epic Games, Inc. All Rights Reserved.
+#include "RigidPhysics/RigidShapeInstanceHandle.h"
+
+#if UE_RIGIDPHYSICS_API_ENABLED
+
+#include "RigidPhysics/RigidScene.h"
+
+namespace UE::Physics
+{
+	FRigidShapeInstanceHandle::FRigidShapeInstanceHandle(const FRigidSceneId& InSceneId, const FRigidObjectId& InId)
+		: Id(InId)
+		, SceneId(InSceneId)
+	{
+	}
+
+	const FRigidObjectId& FRigidShapeInstanceHandle::GetId() const
+	{
+		return Id;
+	}
+	const FRigidSceneId& FRigidShapeInstanceHandle::GetSceneId() const
+	{
+		return SceneId;
+	}
+
+	IRigidShapeInstance* FRigidShapeInstanceHandle::PinInternal(const IRigidScene* Scene) const
+	{
+		return Scene->GetShape(Id);
+	}
+} // namespace UE::Physics
+
+#endif // UE_RIGIDPHYSICS_API_ENABLED
diff --git a/Engine/Source/Runtime/Experimental/RigidPhysics/Public/RigidPhysics/Internal/IRigidBody.h b/Engine/Source/Runtime/Experimental/RigidPhysics/Public/RigidPhysics/Internal/IRigidBody.h
--- a/Engine/Source/Runtime/Experimental/RigidPhysics/Public/RigidPhysics/Internal/IRigidBody.h
+++ b/Engine/Source/Runtime/Experimental/RigidPhysics/Public/RigidPhysics/Internal/IRigidBody.h
@@ -58,10 +58,10 @@ namespace UE::Physics
 		UE_INTERNAL RIGIDPHYSICS_API virtual bool IsSleeping() const;
 
 		UE_INTERNAL virtual int32 GetNumShapes() const = 0;
-		UE_INTERNAL virtual FConstRigidShapeInstancePtr GetShape(int32 ShapeIndex) const = 0;
-
-		// Add a Collision Shape
-		UE_INTERNAL RIGIDPHYSICS_API virtual void AddShape(const FConstRigidShapeInstancePtr& InShape);
+		UE_INTERNAL virtual IRigidShapeInstance* GetShape(int32 ShapeIndex) const = 0;
+		UE_INTERNAL RIGIDPHYSICS_API virtual IRigidShapeInstance* CreateShape(const FRigidShapeInstanceSetup& InSetup);
+		UE_INTERNAL UE_RIGIDPHYSICS_LEGACY_API RIGIDPHYSICS_API virtual IRigidShapeInstance* CreateShape(TUniquePtr<Chaos::FPerShapeData>&& InShapeData);
+		UE_INTERNAL RIGIDPHYSICS_API virtual void DestroyShape(IRigidShapeInstance* InShape, bool bWakeTouching);
 
 		// Set the transform and initialize body state (velocity etc.) at the new location.
 		UE_INTERNAL RIGIDPHYSICS_API virtual void InitTransform(const FTransform3d& InTransform);
diff --git a/Engine/Source/Runtime/Experimental/RigidPhysics/Public/RigidPhysics/Internal/IRigidScene.h b/Engine/Source/Runtime/Experimental/RigidPhysics/Public/RigidPhysics/Internal/IRigidScene.h
--- a/Engine/Source/Runtime/Experimental/RigidPhysics/Public/RigidPhysics/Internal/IRigidScene.h
+++ b/Engine/Source/Runtime/Experimental/RigidPhysics/Public/RigidPhysics/Internal/IRigidScene.h
@@ -41,6 +41,8 @@ namespace UE::Physics
 		UE_INTERNAL RIGIDPHYSICS_API virtual void RegisterModifier(IRigidSceneModifier* InModifier);
 		UE_INTERNAL RIGIDPHYSICS_API virtual void UnregisterModifier(IRigidSceneModifier* InModifier);
 
+		UE_INTERNAL virtual IRigidShapeInstance* GetShape(const FRigidObjectId& InId) const = 0;
+
 		UE_INTERNAL RIGIDPHYSICS_API virtual IRigidBody* CreateBody(const FRigidDebugName& InName, ERigidMovementType InMovementType);
 		UE_INTERNAL RIGIDPHYSICS_API virtual void DestroyBody(IRigidBody* InBody);
 
diff --git a/Engine/Source/Runtime/Experimental/RigidPhysics/Public/RigidPhysics/Internal/IRigidShapeInstance.h b/Engine/Source/Runtime/Experimental/RigidPhysics/Public/RigidPhysics/Internal/IRigidShapeInstance.h
new file mode 100644
--- /dev/null
+++ b/Engine/Source/Runtime/Experimental/RigidPhysics/Public/RigidPhysics/Internal/IRigidShapeInstance.h
@@ -0,0 +1,39 @@
+// Copyright Epic Games, Inc. All Rights Reserved.
+#pragma once
+
+#include "RigidPhysics/RigidFwd.h"
+
+#if UE_RIGIDPHYSICS_API_ENABLED
+
+#include "Chaos/CollisionFilterData.h"
+#include "Math/Transform.h"
+#include "RigidPhysics/Geometry/AnyGeometry.h"
+#include "RigidPhysics/RigidShapeInstanceHandle.h"
+#include "RigidPhysics/RigidLog.h"
+#include "RigidPhysics/RigidTyped.h"
+
+UE_EXPERIMENTAL(5.8, "The new Chaos API is experimental")
+namespace UE::Physics
+{
+	class UE_INTERNAL IRigidShapeInstance : public IRigidTyped
+	{
+	public:
+		UE_RIGIDPHYSICS_RIGIDTYPED_DECL(RIGIDPHYSICS_API, IRigidShapeInstance);
+
+		IRigidShapeInstance() = default;
+		virtual ~IRigidShapeInstance() = default;
+
+		UE_INTERNAL virtual FRigidShapeInstanceHandle GetHandle() const = 0;
+
+		UE_INTERNAL virtual FAnyGeometry GetGeometry() const = 0;
+		// TODO: SetGeometry
+
+		UE_INTERNAL virtual FTransform3f GetLocalTransform() const = 0;
+		// TODO: SetLocalTransform
+
+		UE_INTERNAL virtual Chaos::Filter::FShapeFilterData GetShapeFilter() const = 0;
+		UE_INTERNAL virtual void SetShapeFilter(const Chaos::Filter::FShapeFilterData& ShapeFilter) = 0;
+	};
+} // namespace UE::Physics
+
+#endif // UE_RIGIDPHYSICS_API_ENABLED
diff --git a/Engine/Source/Runtime/Experimental/RigidPhysics/Public/RigidPhysics/RigidBody.h b/Engine/Source/Runtime/Experimental/RigidPhysics/Public/RigidPhysics/RigidBody.h
--- a/Engine/Source/Runtime/Experimental/RigidPhysics/Public/RigidPhysics/RigidBody.h
+++ b/Engine/Source/Runtime/Experimental/RigidPhysics/Public/RigidPhysics/RigidBody.h
@@ -108,21 +108,24 @@ namespace UE::Physics
 			return BodyRaw->GetNumShapes();
 		}
 
-		FConstRigidShapeInstancePtr GetShape(int32 InShapeIndex) const
+		TRigidShapeInstancePtr<FContext> GetShape(int32 InShapeIndex) const
 		{
 			return BodyRaw->GetShape(InShapeIndex);
 		}
 
-		void CreateShape(const FConstRigidShapeInstancePtr& InShape)
+		TRigidShapeInstancePtr<FContext> CreateShape(const FRigidShapeInstanceSetup& InSetup)
 		{
-			BodyRaw->AddShape(InShape);
+			return BodyRaw->CreateShape(InSetup);
 		}
 
-		UE_RIGIDPHYSICS_LEGACY_API FConstRigidShapeInstancePtr CreateShape(TUniquePtr<Chaos::FPerShapeData>&& InShapeData)
+		UE_RIGIDPHYSICS_LEGACY_API TRigidShapeInstancePtr<FContext> CreateShape(TUniquePtr<Chaos::FPerShapeData>&& InShapeData)
 		{
-			FConstRigidShapeInstancePtr Ptr(InShapeData.Release());
-			CreateShape(Ptr);
-			return Ptr;
+			return BodyRaw->CreateShape(MoveTemp(InShapeData));
+		}
+
+		void DestroyShape(const TRigidShapeInstancePtr<FContext>& InShape, bool bWakeTouching = true)
+		{
+			BodyRaw->DestroyShape(InShape.Get(), bWakeTouching);
 		}
 
 		void InitTransform(const FTransform3d& InTransform)
diff --git a/Engine/Source/Runtime/Experimental/RigidPhysics/Public/RigidPhysics/RigidBodyHandle.h b/Engine/Source/Runtime/Experimental/RigidPhysics/Public/RigidPhysics/RigidBodyHandle.h
--- a/Engine/Source/Runtime/Experimental/RigidPhysics/Public/RigidPhysics/RigidBodyHandle.h
+++ b/Engine/Source/Runtime/Experimental/RigidPhysics/Public/RigidPhysics/RigidBodyHandle.h
@@ -40,6 +40,7 @@ namespace UE::Physics
 		}
 
 		UE_INTERNAL RIGIDPHYSICS_API const FRigidObjectId& GetId() const;
+		UE_INTERNAL RIGIDPHYSICS_API const FRigidSceneId& GetSceneId() const;
 		UE_INTERNAL RIGIDPHYSICS_API void Reset();
 
 	private:
diff --git a/Engine/Source/Runtime/Experimental/RigidPhysics/Public/RigidPhysics/RigidFwd.h b/Engine/Source/Runtime/Experimental/RigidPhysics/Public/RigidPhysics/RigidFwd.h
--- a/Engine/Source/Runtime/Experimental/RigidPhysics/Public/RigidPhysics/RigidFwd.h
+++ b/Engine/Source/Runtime/Experimental/RigidPhysics/Public/RigidPhysics/RigidFwd.h
@@ -68,6 +68,8 @@ namespace UE::Physics
 	class FRigidSceneHandle;
 	class FRigidSceneId;
 	class FRigidShapeInstance;
+	class FRigidShapeInstanceHandle;
+	class FRigidShapeInstanceSetup;
 
 	class IJointConstraint;
 	class IJointConstraint6DOF;
@@ -80,6 +82,7 @@ namespace UE::Physics
 	class IRigidScene;
 	class IRigidSceneModifier;
 	class IRigidSceneSettings;
+	class IRigidShapeInstance;
 	class IRigidTyped;
 
 	template<typename ContextType> class TJointConstraintPtrImpl;
@@ -89,7 +92,7 @@ namespace UE::Physics
 	template<typename ContextType> class TRigidGeometryCollectionPtrImpl;
 	template<typename PtrImplType> class TRigidObjectPtr;
 	template<typename ContextType> class TRigidScenePtrImpl;
-	template<bool bConstParam> class TRigidShapeInstancePtr;
+	template<typename ContextType> class TRigidShapeInstancePtrImpl;
 
 	template<typename ContextType> using TJointConstraintPtr = TRigidObjectPtr<TJointConstraintPtrImpl<ContextType>>;
 	template<typename ContextType> using TJointConstraint6DOFPtr = TRigidObjectPtr<TJointConstraint6DOFPtrImpl<ContextType>>;
@@ -97,9 +100,8 @@ namespace UE::Physics
 	template<typename ContextType> using TRigidBodyPtr = TRigidObjectPtr<TRigidBodyPtrImpl<ContextType>>;
 	template<typename ContextType> using TRigidGeometryCollectionPtr = TRigidObjectPtr<TRigidGeometryCollectionPtrImpl<ContextType>>;
 	template<typename ContextType> using TRigidScenePtr = TRigidObjectPtr<TRigidScenePtrImpl<ContextType>>;
+	template<typename ContextType> using TRigidShapeInstancePtr = TRigidObjectPtr<TRigidShapeInstancePtrImpl<ContextType>>;
 
-	using FConstRigidShapeInstancePtr = TRigidShapeInstancePtr<true>;
-	using FRigidShapeInstancePtr = TRigidShapeInstancePtr<false>;
 	using FRigidDebugName = ::Chaos::FSharedDebugName;
 	using FRigidSceneWeakPtr = TWeakPtr<IRigidScene, ESPMode::ThreadSafe>;
 }
diff --git a/Engine/Source/Runtime/Experimental/RigidPhysics/Public/RigidPhysics/RigidShapeInstance.h b/Engine/Source/Runtime/Experimental/RigidPhysics/Public/RigidPhysics/RigidShapeInstance.h
--- a/Engine/Source/Runtime/Experimental/RigidPhysics/Public/RigidPhysics/RigidShapeInstance.h
+++ b/Engine/Source/Runtime/Experimental/RigidPhysics/Public/RigidPhysics/RigidShapeInstance.h
@@ -4,113 +4,101 @@
 #include "RigidPhysics/RigidFwd.h"
 
 #if UE_RIGIDPHYSICS_API_ENABLED
-
-#include "Chaos/ImplicitFwd.h"
 #include "Math/Bounds.h"
 #include "Math/Transform.h"
-#include "RigidPhysics/RigidLog.h"
-#include "RigidPhysics/RigidTyped.h"
-
-class UPhysicalMaterial;
-
-namespace Chaos
-{
-	class FPerShapeData;
-}
+#include "RigidPhysics/Internal/IRigidShapeInstance.h"
+#include "RigidPhysics/RigidObjectId.h"
+#include "RigidPhysics/RigidObjectPtr.h"
+#include "RigidPhysics/RigidShapeInstanceHandle.h"
+#include "RigidPhysics/RigidShapeInstanceSetup.h"
 
+UE_EXPERIMENTAL(5.8, "The new Chaos API is experimental")
 namespace UE::Physics
 {
-	// TODO_CHAOS: Should this be an interface with scene-specific implementation?
-	// Not if we want to use this data type in assets/bodyinstance etc. but it would
-	// allow us to use FPerShapeData and FImplicitObject in the ChaosAsync but move
-	// away from them in new implementations.
-	class UE_INTERNAL FRigidShapeInstance
+	template <typename ContextType>
+	class TRigidShapeInstancePtrImpl
 	{
 	public:
-		UE_INTERNAL RIGIDPHYSICS_API FRigidShapeInstance();
-		UE_INTERNAL RIGIDPHYSICS_API virtual ~FRigidShapeInstance();
-
-		// Existing Chaos implementation exposed for internal use only
-		UE_INTERNAL RIGIDPHYSICS_API FRigidShapeInstance(Chaos::FPerShapeData* InPerShapeData);
-		UE_INTERNAL RIGIDPHYSICS_API Chaos::FPerShapeData* GetPerShapeData() const;
-		UE_INTERNAL RIGIDPHYSICS_API Chaos::FConstImplicitObjectPtr GetImplicitObject() const;
-
-	private:
-		// NOTE: Danger...we hold a pointer to a shape which is the current implementation, 
-		// but its lifetime is currently managed by the Proxy in a UniquePtr.
-		// See FRigidBodyAsyncGT::AddShape
-		Chaos::FPerShapeData* PerShapeData;
-	};
+		using FContext = ContextType;
+		using FInterface = IRigidShapeInstance;
+		using FHandle = FRigidShapeInstanceHandle;
 
+		UE_INTERNAL TRigidShapeInstancePtrImpl() = default;
+		UE_INTERNAL TRigidShapeInstancePtrImpl(IRigidShapeInstance* InShapeInstance)
+			: ShapeInstanceRaw(InShapeInstance)
+		{
+		}
 
-	//
-	// A pointer to a collider. Eventually this will just be the equivalent of a raw pointer
-	// to FRigidShapeInstance, but for now we contain the FRigidShapeInstance by value. This is so that
-	// we don't need to store an array of FRigidShapeInstance objects in our FRigidBody class until
-	// we move away from it being a thin wrapper around FPerShapeData (the particle class
-	// already has an array of PerShapeData so we don't want to double up).
-	//
-	template<bool bConstParam>
-	class TRigidShapeInstancePtr
-	{
-	public:
-		static const bool bIsPtrToConst = bConstParam;
+		TRigidShapeInstancePtrImpl(const TRigidShapeInstancePtrImpl&) = delete;
 
-		TRigidShapeInstancePtr()
-			: ShapeInstance(nullptr)
+		TRigidShapeInstancePtrImpl(TRigidShapeInstancePtrImpl&& InPtr)
+			: ShapeInstanceRaw(InPtr.ShapeInstanceRaw)
 		{
+			InPtr.ShapeInstanceRaw = nullptr;
 		}
 
-		TRigidShapeInstancePtr(const TRigidShapeInstancePtr& R)
-			: ShapeInstance(R.Get()->GetPerShapeData())
+		~TRigidShapeInstancePtrImpl()
 		{
+			Reset();
 		}
 
-		// Auto-cast non-const to const ptr
-		TRigidShapeInstancePtr(const TRigidShapeInstancePtr<false>& R) requires (bIsPtrToConst)
-			: ShapeInstance(R.Get()->GetPerShapeData())
+		friend bool operator==(const TRigidShapeInstancePtrImpl&, const TRigidShapeInstancePtrImpl&) = default;
+		friend bool operator!=(const TRigidShapeInstancePtrImpl&, const TRigidShapeInstancePtrImpl&) = default;
+
+		bool IsValid() const
 		{
+			return (ShapeInstanceRaw != nullptr);
 		}
 
-		bool IsValid() const
+		FRigidShapeInstanceHandle GetHandle() const
 		{
-			return ShapeInstance.GetPerShapeData() != nullptr;
+			if (ShapeInstanceRaw != nullptr)
+			{
+				return ShapeInstanceRaw->GetHandle();
+			}
+			return FRigidShapeInstanceHandle();
 		}
 
-		void Reset()
+		FAnyGeometry GetGeometry() const
 		{
-			ShapeInstance = FRigidShapeInstance();
+			return ShapeInstanceRaw->GetGeometry();
 		}
+		// TODO: SetGeometry
 
-		operator bool() const
+		FTransform3f GetLocalTransform() const
 		{
-			return IsValid();
+			return ShapeInstanceRaw->GetLocalTransform();
 		}
+		// TODO: SetLocalTransform
 
-		FRigidShapeInstance* operator->() const requires (!bIsPtrToConst)
+		Chaos::Filter::FShapeFilterData GetShapeFilter() const
 		{
-			return &ShapeInstance;
+			return ShapeInstanceRaw->GetShapeFilter();
 		}
 
-		const FRigidShapeInstance* operator->() const  requires (bIsPtrToConst)
-		{ 
-			return &ShapeInstance;
+		void SetShapeFilter(const Chaos::Filter::FShapeFilterData& ShapeFilter)
+		{
+			ShapeInstanceRaw->SetShapeFilter(ShapeFilter);
 		}
 
+		UE_INTERNAL void Reset()
+		{
+			ShapeInstanceRaw = nullptr;
+		}
 
-		UE_INTERNAL TRigidShapeInstancePtr(Chaos::FPerShapeData* InShapeInstance)
-			: ShapeInstance(InShapeInstance)
+		UE_INTERNAL const FRigidTypeId& GetTypeId() const
 		{
+			return ShapeInstanceRaw->GetTypeId();
 		}
 
-		UE_INTERNAL const FRigidShapeInstance* Get() const
+		UE_INTERNAL IRigidShapeInstance* Get() const
 		{
-			return &ShapeInstance;
+			return ShapeInstanceRaw;
 		}
 
 	private:
-		mutable FRigidShapeInstance ShapeInstance;
+		IRigidShapeInstance* ShapeInstanceRaw = nullptr;
 	};
-}
+} // namespace UE::Physics
 
-#endif
+#endif // UE_RIGIDPHYSICS_API_ENABLED
diff --git a/Engine/Source/Runtime/Experimental/RigidPhysics/Public/RigidPhysics/RigidShapeInstanceHandle.h b/Engine/Source/Runtime/Experimental/RigidPhysics/Public/RigidPhysics/RigidShapeInstanceHandle.h
new file mode 100644
--- /dev/null
+++ b/Engine/Source/Runtime/Experimental/RigidPhysics/Public/RigidPhysics/RigidShapeInstanceHandle.h
@@ -0,0 +1,45 @@
+// Copyright Epic Games, Inc. All Rights Reserved.
+#pragma once
+
+#include "RigidPhysics/RigidFwd.h"
+
+#if UE_RIGIDPHYSICS_API_ENABLED
+
+#include "RigidPhysics/RigidObjectId.h"
+#include "RigidPhysics/RigidObjectPtr.h"
+#include "RigidPhysics/RigidSceneHandle.h"
+
+UE_EXPERIMENTAL(5.8, "The new Chaos API is experimental")
+namespace UE::Physics
+{
+	class FRigidShapeInstanceHandle
+	{
+	public:
+		FRigidShapeInstanceHandle() = default;
+		FRigidShapeInstanceHandle(const FRigidShapeInstanceHandle&) = default;
+		UE_INTERNAL RIGIDPHYSICS_API FRigidShapeInstanceHandle(const FRigidSceneId& InSceneId, const FRigidObjectId& InId);
+
+		friend bool operator==(const FRigidShapeInstanceHandle&, const FRigidShapeInstanceHandle&) = default;
+		friend bool operator!=(const FRigidShapeInstanceHandle&, const FRigidShapeInstanceHandle&) = default;
+
+		template <typename ContextType>
+		TRigidShapeInstancePtr<ContextType> Pin(const ContextType& InContext) const
+		{
+			if (const IRigidScene* Scene = InContext.GetRawSceneChecked(SceneId))
+			{
+				return PinInternal(Scene);
+			}
+			return TRigidShapeInstancePtr<ContextType>();
+		}
+
+		UE_INTERNAL RIGIDPHYSICS_API const FRigidObjectId& GetId() const;
+		UE_INTERNAL RIGIDPHYSICS_API const FRigidSceneId& GetSceneId() const;
+	private:
+		RIGIDPHYSICS_API IRigidShapeInstance* PinInternal(const IRigidScene* Scene) const;
+
+		FRigidObjectId Id;
+		FRigidSceneId SceneId;
+	};
+} // namespace UE::Physics
+
+#endif // UE_RIGIDPHYSICS_API_ENABLED
diff --git a/Engine/Source/Runtime/Experimental/RigidPhysics/Public/RigidPhysics/RigidShapeInstanceSetup.h b/Engine/Source/Runtime/Experimental/RigidPhysics/Public/RigidPhysics/RigidShapeInstanceSetup.h
new file mode 100644
--- /dev/null
+++ b/Engine/Source/Runtime/Experimental/RigidPhysics/Public/RigidPhysics/RigidShapeInstanceSetup.h
@@ -0,0 +1,26 @@
+// Copyright Epic Games, Inc. All Rights Reserved.
+#pragma once
+
+#include "RigidPhysics/RigidFwd.h"
+
+#if UE_RIGIDPHYSICS_API_ENABLED
+#include "Math/Transform.h"
+#include "RigidPhysics/Geometry/AnyGeometry.h"
+
+UE_EXPERIMENTAL(5.8, "The new Chaos API is experimental")
+namespace UE::Physics
+{
+	class FRigidShapeInstanceSetup
+	{
+	public:
+		FRigidShapeInstanceSetup(const FAnyGeometry& InGeometry)
+			: Geometry(InGeometry)
+		{
+		}
+
+		FAnyGeometry Geometry;
+		FTransform3f LocalTransform;
+	};
+} // namespace UE::Physics
+
+#endif // UE_RIGIDPHYSICS_API_ENABLED
