Chaos API: Implementing capsule for shape instance.
#Synced-CL 51413247

diff --git a/Engine/Plugins/Experimental/ChaosRigidPhysicsAsync/Source/ChaosRigidPhysicsAsync/Private/ChaosRigidPhysicsAsync/RigidShapeInstanceAsync.cpp b/Engine/Plugins/Experimental/ChaosRigidPhysicsAsync/Source/ChaosRigidPhysicsAsync/Private/ChaosRigidPhysicsAsync/RigidShapeInstanceAsync.cpp
--- a/Engine/Plugins/Experimental/ChaosRigidPhysicsAsync/Source/ChaosRigidPhysicsAsync/Private/ChaosRigidPhysicsAsync/RigidShapeInstanceAsync.cpp
+++ b/Engine/Plugins/Experimental/ChaosRigidPhysicsAsync/Source/ChaosRigidPhysicsAsync/Private/ChaosRigidPhysicsAsync/RigidShapeInstanceAsync.cpp
@@ -36,7 +36,11 @@ namespace Chaos::Rigids::Async
 			}
 			else if constexpr (std::is_same_v<GeometryType, UE::Physics::FCapsuleGeometry>)
 			{
-				ensureMsgf(false, TEXT("TODO"));
+				const float Radius = Geometry.GetRadius();
+				const float Length = Geometry.GetSegmentHalfLength() * 2;
+				FKSphylElem CapsuleElement = FKSphylElem(Radius, Length);
+				CapsuleElement.Center = FVector::Zero();
+				AggGeom.AddElement(CapsuleElement);
 			}
 			else if constexpr (std::is_same_v<GeometryType, UE::Physics::FConvexGeometry>)
 			{
@@ -103,6 +107,7 @@ namespace Chaos::Rigids::Async
 		using FImplicitTransformed = TImplicitObjectTransformed<FReal, 3>;
 		using FBoxGeometry = UE::Physics::FBoxGeometry;
 		using FSphereGeometry = UE::Physics::FSphereGeometry;
+		using FCapsuleGeometry = UE::Physics::FCapsuleGeometry;
 		using FAnyGeometry = UE::Physics::FAnyGeometry;
 
 		const FImplicitObject* ImplicitGeom = PerShapeData->GetGeometry();
@@ -125,6 +130,11 @@ namespace Chaos::Rigids::Async
 			{
 				return FBoxGeometry(FVector3f(ImplicitGeom->GetObjectChecked<FImplicitBox3>().Extents()));
 			}
+			case ImplicitObjectType::Capsule:
+			{
+				const FImplicitCapsule3& Capsule = ImplicitGeom->GetObjectChecked<FImplicitCapsule3>();
+				return FCapsuleGeometry(Capsule.GetRadiusf(), Capsule.GetHeightf() / 2.0f);
+			}
 			default:
 			{
 				ensureMsgf(false, TEXT("TODO"));
@@ -151,6 +161,12 @@ namespace Chaos::Rigids::Async
 				FImplicitBox3& Box = Geom->GetObjectChecked<FImplicitBox3>();
 				return FTransform3f(FVector3f(Box.Center()));
 			}
+			case ImplicitObjectType::Capsule:
+			{
+				FImplicitCapsule3& Capsule = Geom->GetObjectChecked<FImplicitCapsule3>();
+				FQuat4f Rotation = FQuat4f::FindBetween(FVector3f(0, 0, 1), Capsule.GetAxisf());
+				return FTransform3f(Rotation, Capsule.GetCenterf());
+			}
 			case ImplicitObjectType::Transformed:
 			{
 				const FImplicitTransformed& Transformed = Geom->GetObjectChecked<FImplicitTransformed>();
diff --git a/Engine/Plugins/Experimental/ChaosRigidPhysicsAsync/Tests/Private/AnyGeometryTests.cpp b/Engine/Plugins/Experimental/ChaosRigidPhysicsAsync/Tests/Private/AnyGeometryTests.cpp
--- a/Engine/Plugins/Experimental/ChaosRigidPhysicsAsync/Tests/Private/AnyGeometryTests.cpp
+++ b/Engine/Plugins/Experimental/ChaosRigidPhysicsAsync/Tests/Private/AnyGeometryTests.cpp
@@ -12,6 +12,7 @@ namespace Chaos::LowLevelTest
 {
 	using FSphereGeometry = UE::Physics::FSphereGeometry;
 	using FBoxGeometry = UE::Physics::FBoxGeometry;
+	using FCapsuleGeometry = UE::Physics::FCapsuleGeometry;
 	using FAnyGeometry = UE::Physics::FAnyGeometry;
 
 	template <typename ExpectedGeometryType, typename ActualGeometryType>
@@ -124,6 +125,52 @@ namespace Chaos::LowLevelTest
 			CHECK(AnyGeometry == Sphere1);
 		}
 	}
+
+	TEST_CASE("AnyGeometry: Capsule", "[Chaos][API][Geometry][unit]")
+	{
+		FBoxGeometry Box0(FVector3f(1, 2, 3));
+		FCapsuleGeometry Capsule0(1, 3);
+		FCapsuleGeometry Capsule1(2, 4);
+
+		SECTION("IsA")
+		{
+			FAnyGeometry AnyGeometry(Capsule0);
+			CHECK(AnyGeometry.IsA<FCapsuleGeometry>());
+			CHECK(!AnyGeometry.IsA<FBoxGeometry>());
+		}
+		SECTION("TryGet/Get")
+		{
+			FAnyGeometry AnyGeometry(Capsule0);
+			REQUIRE(AnyGeometry.TryGet<FCapsuleGeometry>() != nullptr);
+			CHECK(*AnyGeometry.TryGet<FCapsuleGeometry>() == Capsule0);
+			CHECK(AnyGeometry.Get<FCapsuleGeometry>() == Capsule0);
+			CHECK(AnyGeometry.TryGet<FBoxGeometry>() == nullptr);
+		}
+		SECTION("Visit")
+		{
+			FAnyGeometry AnyGeometry(Capsule0);
+			AnyGeometry.Visit([&Capsule0]<typename GeometryType>(const GeometryType& Geometry) { CheckGeometry(Capsule0, Geometry); });
+		}
+		SECTION("Comparison")
+		{
+			FAnyGeometry AnyGeometry(Capsule0);
+			CHECK(AnyGeometry == Capsule0);
+			CHECK(AnyGeometry != Capsule1);
+			CHECK(AnyGeometry != Box0);
+		}
+		SECTION("Assignment")
+		{
+			FAnyGeometry AnyGeometry(Box0);
+			AnyGeometry = Capsule0;
+			CHECK(AnyGeometry != Box0);
+			CHECK(AnyGeometry == Capsule0);
+			CHECK(AnyGeometry != Capsule1);
+			AnyGeometry = Capsule1;
+			CHECK(AnyGeometry != Box0);
+			CHECK(AnyGeometry != Capsule0);
+			CHECK(AnyGeometry == Capsule1);
+		}
+	}
 } // namespace Chaos::LowLevelTest
 
 #endif // UE_RIGIDPHYSICS_API_ENABLED
diff --git a/Engine/Plugins/Experimental/ChaosRigidPhysicsAsync/Tests/Private/ShapeInstanceTests.cpp b/Engine/Plugins/Experimental/ChaosRigidPhysicsAsync/Tests/Private/ShapeInstanceTests.cpp
--- a/Engine/Plugins/Experimental/ChaosRigidPhysicsAsync/Tests/Private/ShapeInstanceTests.cpp
+++ b/Engine/Plugins/Experimental/ChaosRigidPhysicsAsync/Tests/Private/ShapeInstanceTests.cpp
@@ -15,6 +15,19 @@ namespace Chaos::LowLevelTest
 {
 	using namespace UE::Physics;
 
+	template <typename GeometryType>
+	void CheckGeometry(const GeometryType& Expected, const GeometryType& Actual)
+	{
+		CHECK(Expected == Actual);
+	}
+
+	template <>
+	void CheckGeometry<FCapsuleGeometry>(const FCapsuleGeometry& Expected, const FCapsuleGeometry& Actual)
+	{
+		CHECK_THAT(Expected.GetRadius(), Catch::Matchers::WithinAbs(Actual.GetRadius(), UE_KINDA_SMALL_NUMBER));
+		CHECK_THAT(Expected.GetSegmentHalfLength(), Catch::Matchers::WithinAbs(Actual.GetSegmentHalfLength(), UE_KINDA_SMALL_NUMBER));
+	}
+
 	template <typename GeometryType>
 	void CheckGeometryTransform(const FTransform3f& Expected, const FTransform3f& Actual)
 	{
@@ -24,6 +37,22 @@ namespace Chaos::LowLevelTest
 		CHECK_THAT(Chaos::FVec3f(Expected.GetScale3D()), Catch::ApproxEq(Chaos::FVec3f(Actual.GetScale3D())));
 	}
 
+	template <>
+	void CheckGeometryTransform<FCapsuleGeometry>(const FTransform3f& Expected, const FTransform3f& Actual)
+	{
+		CHECK_THAT(Chaos::FVec3f(Expected.GetLocation()), Catch::ApproxEq(Chaos::FVec3f(Actual.GetLocation())));
+		CHECK_THAT(Chaos::FVec3f(Expected.GetScale3D()), Catch::ApproxEq(Chaos::FVec3f(Actual.GetScale3D())));
+
+		// Capsule rotations are deconstructed from an axis. There is a many to one mapping for a capsule axis, so we cannot check the rotation directly.
+		// Instead, we want to verify the "pose" of the capsule, we can do this by transforming a capsule axis and verifying it's the same.
+		const FVec3f ExpectedAxis = Expected.GetRotation().RotateVector(FVec3f(0, 0, 1)).GetSafeNormal();
+		FVec3f ActualAxis = Actual.GetRotation().RotateVector(FVec3f(0, 0, 1)).GetSafeNormal();
+
+		constexpr float AngularErrorThreshold = FMath::DegreesToRadians(1);
+		float AngularDistanceRad = FMath::Abs(FMath::Acos(ExpectedAxis.Dot(ActualAxis)));
+		CHECK(AngularDistanceRad < AngularErrorThreshold);
+	}
+
 	void CreateAndCheckShapeInstanceAndTransform(FRigidTestFixture& Fixture, const FRigidShapeInstanceSetup& InSetup, const FAnyGeometry& ExpectedGeometry, const FTransform3f& ExpectedTransform)
 	{
 		if (FRigidContextGameRW Context = Fixture.LockRWChecked())
@@ -35,11 +64,14 @@ namespace Chaos::LowLevelTest
 
 			const FTransform3f ActualTransform = ShapeInstance->GetLocalTransform();
 			const FAnyGeometry ActualGeometry = ShapeInstance->GetGeometry();
-			CHECK(ExpectedGeometry == ActualGeometry);
-			// Visit the actual geometry type so we can control how we check the transform.
-			// Some types don't have a unique transform, so instead we have to test the "pose" of the transform.
-			ActualGeometry.Visit([&ExpectedTransform, &ActualTransform]<typename GeometryType>(const GeometryType& ActualGeometry)
+
+			ActualGeometry.Visit([&ExpectedTransform, &ActualTransform, &ExpectedGeometry]<typename GeometryType>(const GeometryType& ActualTypedGeometry)
 			{
+				const GeometryType* ExpectedTypedGeometry = ExpectedGeometry.TryGet<GeometryType>();
+				CHECKED_IF(ExpectedTypedGeometry)
+				{
+					CheckGeometry(*ExpectedTypedGeometry, ActualTypedGeometry);
+				}
 				CheckGeometryTransform<GeometryType>(ExpectedTransform, ActualTransform);
 			});
 		}
@@ -128,6 +160,74 @@ namespace Chaos::LowLevelTest
 				CreateAndCheckShapeInstanceAndTransform(Fixture, Setup, FSphereGeometry(6), ExpectedTransform);
 			}
 		}
+		SECTION("Capsule")
+		{
+			FRigidShapeInstanceSetup Setup(FCapsuleGeometry(2, 3));
+
+			SECTION("Identity Transform")
+			{
+				Setup.LocalTransform.SetIdentity();
+				CreateAndCheckShapeInstanceAndTransform(Fixture, Setup);
+			}
+			SECTION("Translation")
+			{
+				Setup.LocalTransform.SetTranslation(FVector3f(0, 1, 0));
+				CreateAndCheckShapeInstanceAndTransform(Fixture, Setup);
+			}
+			SECTION("Rotation")
+			{
+				const float Degrees = GENERATE(90, 45);
+				const FVector3f Axis = GENERATE(
+					FVector3f(1, 0, 0),
+					FVector3f(0, 1, 0),
+					FVector3f(0, 0, 1),
+					FVector3f(1, 1, 0),
+					FVector3f(1, 0, 1),
+					FVector3f(0, 1, 1),
+					FVector3f(1, 1, 1)
+				);
+				const FQuat4f Rotation = FQuat4f(Axis, FMath::DegreesToRadians(Degrees));
+				CreateAndCheckShapeInstanceAndTransform(Fixture, Setup, Setup.Geometry, Setup.LocalTransform);
+			}
+			SECTION("Scale")
+			{
+				struct FScaleCase
+				{
+					FScaleCase(const FVector3f& InScale, const FCapsuleGeometry& InGeometry) : Scale(InScale), ExpectedGeometry(InGeometry) {}
+					FVector3f Scale;
+					FCapsuleGeometry ExpectedGeometry;
+				};
+				// Capsule scaling is a bit difficult to understand without considering the full picture.
+				// When a capsule is scaled, the full extents are what's actually scaled then the capsule is built from that.
+				// This can be thought of as scaling the cylinder that fully contains the capsule.
+				// So Capsule(Radius, HalfHeight) -> Cylinder(Radius, FullHalfHeight): Cylinder.FullHalfHeight = Capsule.Radius + Cylinder.HalfHeight.
+				// So an example:
+				// - Capsule(2, 3) with scale (2, 1, 1)
+				//   - Capsule(2, 3) -> Cylinder(2, 3 + 2) -> Cylinder(2, 5)
+				//   - So Cylinder(2, 5) * Scale(2, 1) -> Cylinder (4, 5)
+				//   - Cylinder(4, 5) -> Capsule(4, 5 - 4) -> Capsule(4, 1)
+				FScaleCase CaseData = GENERATE(
+					FScaleCase(FVector3f(2, 1, 1), FCapsuleGeometry(4, 1)),
+					FScaleCase(FVector3f(1, 2, 1), FCapsuleGeometry(4, 1)),
+					// Capsule(2, 3) * (1, 1, 2) -> Cylinder(2, 5) * (1, 2) -> Cylinder(2, 10) -> Capsule(2, 10 - 2) -> Capsule(2, 8)
+					FScaleCase(FVector3f(1, 1, 2), FCapsuleGeometry(2, 8)),
+					// Capsule(2, 3) * (2, 2, 2) -> Cylinder(2, 5) * (2, 2) -> Cylinder(4, 10) -> Capsule(4, 10 - 4) -> Capsule(2, 6)
+					FScaleCase(FVector3f(2, 2, 2), FCapsuleGeometry(4, 6))
+				);
+				Setup.LocalTransform.SetScale3D(CaseData.Scale);
+				CreateAndCheckShapeInstanceAndTransform(Fixture, Setup, CaseData.ExpectedGeometry, FTransform3f());
+			}
+			SECTION("Full Transform")
+			{
+				Setup.LocalTransform.SetLocation(FVector3f(-1, -2, -3));
+				Setup.LocalTransform.SetRotation(FRotator3f(10, 20, 30).Quaternion());
+				// Capsule(2, 3) * (2, 3, 4) -> Cylinder(2, 5) * (3, 4) -> Cylinder(6, 20) -> Capsule(6, 20 - 6) -> Capsule(6, 14)
+				Setup.LocalTransform.SetScale3D(FVector3f(2, 3, 4));
+				FTransform3f ExpectedTransform = Setup.LocalTransform;
+				ExpectedTransform.SetScale3D(FVec3f::One());
+				CreateAndCheckShapeInstanceAndTransform(Fixture, Setup, FCapsuleGeometry(6, 14), ExpectedTransform);
+			}
+		}
 	}
 
 	TEST_CASE("ShapeInstance: ShapeFilter", "[Chaos][API][ShapeInstance][unit]")
