Chaos API: Adding more tests for shape instance.
Adding more edge cases for box, sphere, and capsule, mostly around zeros and scale (including negative).
#Synced-CL 51759702

diff --git a/Engine/Plugins/Experimental/ChaosRigidPhysicsAsync/Tests/Private/ShapeInstanceTests.cpp b/Engine/Plugins/Experimental/ChaosRigidPhysicsAsync/Tests/Private/ShapeInstanceTests.cpp
--- a/Engine/Plugins/Experimental/ChaosRigidPhysicsAsync/Tests/Private/ShapeInstanceTests.cpp
+++ b/Engine/Plugins/Experimental/ChaosRigidPhysicsAsync/Tests/Private/ShapeInstanceTests.cpp
@@ -21,6 +21,16 @@ namespace Chaos::LowLevelTest
 		CHECK(Expected == Actual);
 	}
 
+	void CheckGeometry(const FBoxGeometry& Expected, const FBoxGeometry& Actual)
+	{
+		CHECK_THAT(Chaos::FVec3f(Expected.GetExtents()), Catch::ApproxEq(Chaos::FVec3f(Actual.GetExtents()), UE_KINDA_SMALL_NUMBER));
+	}
+
+	void CheckGeometry(const FSphereGeometry& Expected, const FSphereGeometry& Actual)
+	{
+		CHECK_THAT(Expected.GetRadius(), Catch::Matchers::WithinAbs(Actual.GetRadius(), UE_KINDA_SMALL_NUMBER));
+	}
+
 	template <>
 	void CheckGeometry<FCapsuleGeometry>(const FCapsuleGeometry& Expected, const FCapsuleGeometry& Actual)
 	{
@@ -106,8 +116,20 @@ namespace Chaos::LowLevelTest
 			}
 			SECTION("Scale")
 			{
-				Setup.LocalTransform.SetScale3D(FVector3f(2, 3, 4));
-				CreateAndCheckShapeInstanceAndTransform(Fixture, Setup, FBoxGeometry(FVector3f(2, 6, 12)), FTransform3f());
+				struct FBoxScaleCase
+				{
+					FBoxScaleCase(const FVector3f& InScale, const FVector3f& InExpectedExtents) : Scale(InScale), ExpectedExtents(InExpectedExtents) {}
+					FVector3f Scale;
+					FVector3f ExpectedExtents;
+				};
+				const FBoxScaleCase CaseData = GENERATE(
+					FBoxScaleCase(FVector3f(2, 3, 4), FVector3f(2, 6, 12)),
+					FBoxScaleCase(FVector3f(-2, -3, -4), FVector3f(2, 6, 12)),
+					FBoxScaleCase(FVector3f(0, 0, 0), FVector3f(.1f, .2f, .3f))
+				);
+				CAPTURE(Chaos::FVec3f(CaseData.Scale));
+				Setup.LocalTransform.SetScale3D(CaseData.Scale);
+				CreateAndCheckShapeInstanceAndTransform(Fixture, Setup, FBoxGeometry(CaseData.ExpectedExtents), FTransform3f());
 			}
 			SECTION("Full Transform")
 			{
@@ -118,11 +140,17 @@ namespace Chaos::LowLevelTest
 				ExpectedTransform.SetScale3D(FVec3f::One());
 				CreateAndCheckShapeInstanceAndTransform(Fixture, Setup, FBoxGeometry(FVector3f(2, 6, 12)), ExpectedTransform);
 			}
+			SECTION("Zero Extents")
+			{
+				Setup = FRigidShapeInstanceSetup(FBoxGeometry(FVector3f::Zero()));
+				// Note: The code sets half-extents to UE_KINDA_SMALL_NUMBER, so the full extents become 2 times that.
+				const FBoxGeometry ExpectedBox(FVector3f(UE_KINDA_SMALL_NUMBER * 2));
+				CreateAndCheckShapeInstanceAndTransform(Fixture, Setup, ExpectedBox, FTransform3f());
+			}
 		}
 		SECTION("Sphere")
 		{
 			FRigidShapeInstanceSetup Setup(FSphereGeometry(2));
-
 			SECTION("Identity Transform")
 			{
 				Setup.LocalTransform.SetIdentity();
@@ -138,17 +166,30 @@ namespace Chaos::LowLevelTest
 				Setup.LocalTransform.SetRotation(FRotator3f(10, 20, 30).Quaternion());
 				CreateAndCheckShapeInstanceAndTransform(Fixture, Setup, Setup.Geometry, FTransform3f());
 			}
-			SECTION("Uniform Scale")
-			{
-				Setup.LocalTransform.SetScale3D(FVector3f(3));
-				CreateAndCheckShapeInstanceAndTransform(Fixture, Setup, FSphereGeometry(6), FTransform3f());
-			}
-			SECTION("Non-Uniform Scale")
+			SECTION("Scale")
 			{
-				// The min scale element is always used
-				const FVector3f Scale = GENERATE(FVector3f(2, 3, 4), FVector3f(3, 4, 2), FVector3f(4, 2, 3));
-				Setup.LocalTransform.SetScale3D(Scale);
-				CreateAndCheckShapeInstanceAndTransform(Fixture, Setup, FSphereGeometry(4), FTransform3f());
+				struct FSphereScaleCase
+				{
+					FSphereScaleCase(const FVector3f& InScale, const float InExpectedRadius) : Scale(InScale), ExpectedRadius(InExpectedRadius) {}
+					FVector3f Scale;
+					float ExpectedRadius;
+				};
+				const FSphereScaleCase CaseData = GENERATE(
+					// Uniform:
+					FSphereScaleCase(FVector3f(3), 6),
+					// Non uniform (positive): min element is used
+					FSphereScaleCase(FVector3f(2, 3, 4), 4),
+					FSphereScaleCase(FVector3f(3, 4, 2), 4),
+					FSphereScaleCase(FVector3f(4, 2, 3), 4),
+					// Zero:
+					FSphereScaleCase(FVector3f::Zero(), 0.2),
+					// Negative:
+					FSphereScaleCase(FVector3f(-2, -3, -4), 4), // All are negative, picks min of abs(scale)
+					FSphereScaleCase(FVector3f(-2, 1, 3), 2) // Picks largest non-negative
+				);
+				CAPTURE(Chaos::FVec3f(CaseData.Scale));
+				Setup.LocalTransform.SetScale3D(CaseData.Scale);
+				CreateAndCheckShapeInstanceAndTransform(Fixture, Setup, FSphereGeometry(CaseData.ExpectedRadius), FTransform3f());
 			}
 			SECTION("Full Transform")
 			{
@@ -159,6 +200,11 @@ namespace Chaos::LowLevelTest
 				ExpectedTransform.SetLocation(Setup.LocalTransform.GetLocation());
 				CreateAndCheckShapeInstanceAndTransform(Fixture, Setup, FSphereGeometry(6), ExpectedTransform);
 			}
+			SECTION("Zero Radius")
+			{
+				Setup = FRigidShapeInstanceSetup(FSphereGeometry(0));
+				CreateAndCheckShapeInstanceAndTransform(Fixture, Setup, FSphereGeometry(0), FTransform3f());
+			}
 		}
 		SECTION("Capsule")
 		{
@@ -191,9 +237,9 @@ namespace Chaos::LowLevelTest
 			}
 			SECTION("Scale")
 			{
-				struct FScaleCase
+				struct FCapsuleScaleCase
 				{
-					FScaleCase(const FVector3f& InScale, const FCapsuleGeometry& InGeometry) : Scale(InScale), ExpectedGeometry(InGeometry) {}
+					FCapsuleScaleCase(const FVector3f& InScale, const FCapsuleGeometry& InGeometry) : Scale(InScale), ExpectedGeometry(InGeometry) {}
 					FVector3f Scale;
 					FCapsuleGeometry ExpectedGeometry;
 				};
@@ -206,14 +252,31 @@ namespace Chaos::LowLevelTest
 				//   - Capsule(2, 3) -> Cylinder(2, 3 + 2) -> Cylinder(2, 5)
 				//   - So Cylinder(2, 5) * Scale(2, 1) -> Cylinder (4, 5)
 				//   - Cylinder(4, 5) -> Capsule(4, 5 - 4) -> Capsule(4, 1)
-				FScaleCase CaseData = GENERATE(
-					FScaleCase(FVector3f(2, 1, 1), FCapsuleGeometry(4, 1)),
-					FScaleCase(FVector3f(1, 2, 1), FCapsuleGeometry(4, 1)),
+				// Additionally, the radius and the half height are clamped to prevent invalid capsules:
+				// - The radius is Clamp(0.1, HalfHeight)
+				// - The half height is Clamp(.05, HalfHeight - Radius)
+				FCapsuleScaleCase CaseData = GENERATE(
+					FCapsuleScaleCase(FVector3f(2, 1, 1), FCapsuleGeometry(4, 1)),
+					FCapsuleScaleCase(FVector3f(1, 2, 1), FCapsuleGeometry(4, 1)),
 					// Capsule(2, 3) * (1, 1, 2) -> Cylinder(2, 5) * (1, 2) -> Cylinder(2, 10) -> Capsule(2, 10 - 2) -> Capsule(2, 8)
-					FScaleCase(FVector3f(1, 1, 2), FCapsuleGeometry(2, 8)),
+					FCapsuleScaleCase(FVector3f(1, 1, 2), FCapsuleGeometry(2, 8)),
 					// Capsule(2, 3) * (2, 2, 2) -> Cylinder(2, 5) * (2, 2) -> Cylinder(4, 10) -> Capsule(4, 10 - 4) -> Capsule(2, 6)
-					FScaleCase(FVector3f(2, 2, 2), FCapsuleGeometry(4, 6))
+					FCapsuleScaleCase(FVector3f(2, 2, 2), FCapsuleGeometry(4, 6)),
+					// Zero:
+					FCapsuleScaleCase(FVector3f::Zero(), FCapsuleGeometry(.2, .3)),
+					// Capsule(2, 3) * (0, 0, 1) -> Cylinder(2, 5) * (0, 1) -> Cylinder(0, 5) -> Cylinder(0.1, 5) -> Capsule(0.1, 5 - 0.1) -> Capsule(0.1, 4.9)
+					FCapsuleScaleCase(FVector3f(0, 0, 1), FCapsuleGeometry(.1, 4.9)),
+					// This case tests radius being clamped to half height. This turns into Capsule(0, 0)
+					FCapsuleScaleCase(FVector3f(1, 1, 0), FCapsuleGeometry(0.1f, 0.05f)),
+					// Negative:
+					FCapsuleScaleCase(FVector3f(-2, -2, -2), FCapsuleGeometry(4, 6)),
+					FCapsuleScaleCase(FVector3f(-2, 1, 1), FCapsuleGeometry(4, 1)), // Same as (2, 1, 1)
+					FCapsuleScaleCase(FVector3f(1, 1, -2), FCapsuleGeometry(2, 8)), // Same as (1, 1, 2)
+					// This case tests radius being clamped to half height:
+					// Capsule(2, 3) * (10, 1) -> Cylinder(2, 5) * (10, 1) -> Cylinder(20, 5) -> (Clamp Radius to HalfHeight) -> Cylinder(5, 5) ->Capsule(5, 5 - 5) -> Capsule(5, 0.1)
+					FCapsuleScaleCase(FVector3f(10, 1, 1), FCapsuleGeometry(5, 0.05f))
 				);
+				CAPTURE(Chaos::FVec3f(CaseData.Scale));
 				Setup.LocalTransform.SetScale3D(CaseData.Scale);
 				CreateAndCheckShapeInstanceAndTransform(Fixture, Setup, CaseData.ExpectedGeometry, FTransform3f());
 			}
@@ -227,38 +290,14 @@ namespace Chaos::LowLevelTest
 				ExpectedTransform.SetScale3D(FVec3f::One());
 				CreateAndCheckShapeInstanceAndTransform(Fixture, Setup, FCapsuleGeometry(6, 14), ExpectedTransform);
 			}
-		}
-	}
-
-	TEST_CASE("ShapeInstance: ShapeFilter", "[Chaos][API][ShapeInstance][unit]")
-	{
-		FRigidTestFixture Fixture;
-		FRigidShapeInstanceSetup Setup(FBoxGeometry(FVector3f(1, 2, 3)));
-
-		const Chaos::Filter::FShapeFilterData ShapeFilter0 = Chaos::Filter::FShapeFilterBuilder().SetCollisionChannelIndex(1).Build();
-		const Chaos::Filter::FShapeFilterData ShapeFilter1 = Chaos::Filter::FShapeFilterBuilder().SetCollisionChannelIndex(2).Build();
-
-		FRigidShapeInstanceHandle ShapeInstanceHandle;
-		if (FRigidContextGameRW Context = Fixture.LockRWChecked())
-		{
-			TRigidBodyPtr<FRigidContextGameRW> BodyPtr = Context->CreateBody(FRigidDebugName("Body1"), ERigidMovementType::Dynamic);
-			Fixture.AutoCleanup(BodyPtr);
-
-			TRigidShapeInstancePtr<FRigidContextGameRW> ShapeInstance0 = BodyPtr->CreateShape(Setup);
-			ShapeInstanceHandle = ShapeInstance0;
-			BodyPtr->Activate();
-
-			ShapeInstance0->SetShapeFilter(ShapeFilter0);
-			CHECK(ShapeInstance0->GetShapeFilter() == ShapeFilter0);
-			ShapeInstance0->SetShapeFilter(ShapeFilter1);
-			CHECK(ShapeInstance0->GetShapeFilter() == ShapeFilter1);
-		}
-
-		Fixture.RunPTCallback([ShapeInstanceHandle, &ShapeFilter1](const FRigidContextSimRW& Context)
+			SECTION("Zero Input")
 			{
-				const TRigidShapeInstancePtr<FRigidContextSimRW> ShapeInstancePtr = ShapeInstanceHandle.Pin(Context);
-				CHECK(ShapeInstancePtr->GetShapeFilter() == ShapeFilter1);
-			});
+				Setup = FRigidShapeInstanceSetup(FCapsuleGeometry(0, 0));
+				// Internals currently clamp radius and height to 0.1.
+				const FCapsuleGeometry Expected(0.1f, 0.05f);
+				CreateAndCheckShapeInstanceAndTransform(Fixture, Setup, Expected, FTransform3f());
+			}
+		}
 	}
 } // namespace Chaos::LowLevelTest
 
