Chaos API: Exposing more creation params for shape instance creation.
#Synced-CL 52100670

diff --git a/Engine/Plugins/Experimental/ChaosRigidPhysicsAsync/Source/ChaosRigidPhysicsAsync/Private/ChaosRigidPhysicsAsync/Conversions.h b/Engine/Plugins/Experimental/ChaosRigidPhysicsAsync/Source/ChaosRigidPhysicsAsync/Private/ChaosRigidPhysicsAsync/Conversions.h
--- a/Engine/Plugins/Experimental/ChaosRigidPhysicsAsync/Source/ChaosRigidPhysicsAsync/Private/ChaosRigidPhysicsAsync/Conversions.h
+++ b/Engine/Plugins/Experimental/ChaosRigidPhysicsAsync/Source/ChaosRigidPhysicsAsync/Private/ChaosRigidPhysicsAsync/Conversions.h
@@ -10,6 +10,28 @@
 
 namespace Chaos::Rigids::Async
 {
+	inline ECollisionTraceFlag ConvertCollisionTraceFlag(const Chaos::EChaosCollisionTraceFlag InFlag)
+	{
+		switch (InFlag)
+		{
+			case Chaos::EChaosCollisionTraceFlag::Chaos_CTF_UseDefault:
+				return ECollisionTraceFlag::CTF_UseDefault;
+			case Chaos::EChaosCollisionTraceFlag::Chaos_CTF_UseSimpleAndComplex:
+				return ECollisionTraceFlag::CTF_UseSimpleAndComplex;
+			case Chaos::EChaosCollisionTraceFlag::Chaos_CTF_UseSimpleAsComplex:
+				return ECollisionTraceFlag::CTF_UseSimpleAsComplex;
+			case Chaos::EChaosCollisionTraceFlag::Chaos_CTF_UseComplexAsSimple:
+				return ECollisionTraceFlag::CTF_UseComplexAsSimple;
+			case Chaos::EChaosCollisionTraceFlag::Chaos_CTF_MAX:
+				return ECollisionTraceFlag::CTF_MAX;
+			default:
+			{
+				ensure(false);
+				return ECollisionTraceFlag::CTF_UseDefault;
+			}
+		}
+	}
+
 	inline Chaos::EChaosCollisionTraceFlag ConvertCollisionTraceFlag(const ECollisionTraceFlag InFlag)
 	{
 		switch (InFlag)
diff --git a/Engine/Plugins/Experimental/ChaosRigidPhysicsAsync/Source/ChaosRigidPhysicsAsync/Private/ChaosRigidPhysicsAsync/RigidShapeInstanceAsync.cpp b/Engine/Plugins/Experimental/ChaosRigidPhysicsAsync/Source/ChaosRigidPhysicsAsync/Private/ChaosRigidPhysicsAsync/RigidShapeInstanceAsync.cpp
--- a/Engine/Plugins/Experimental/ChaosRigidPhysicsAsync/Source/ChaosRigidPhysicsAsync/Private/ChaosRigidPhysicsAsync/RigidShapeInstanceAsync.cpp
+++ b/Engine/Plugins/Experimental/ChaosRigidPhysicsAsync/Source/ChaosRigidPhysicsAsync/Private/ChaosRigidPhysicsAsync/RigidShapeInstanceAsync.cpp
@@ -4,6 +4,7 @@
 #if UE_RIGIDPHYSICS_API_ENABLED
 
 #include "Chaos/Box.h"
+#include "Chaos/CastingUtilities.h"
 #include "Chaos/Capsule.h"
 #include "Chaos/Convex.h"
 #include "Chaos/ImplicitObjectTransformed.h"
@@ -128,16 +129,14 @@ namespace Chaos::Rigids::Async
 		TUniquePtr<Chaos::FPerShapeData> NewShape = Chaos::FShapeInstanceProxy::Make(0, ImplicitPtr);
 		Chaos::FPerShapeData* PerShapeData = NewShape.Release();
 
-		const Chaos::EFilterFlags FilterFlags = Chaos::EFilterFlags::SimpleCollision | Chaos::EFilterFlags::ContactNotify | Chaos::EFilterFlags::ModifyContacts;
-		const Chaos::Filter::FShapeFilterData ShapeFilter = Chaos::Filter::FShapeFilterBuilder::BuildBlockAll(FilterFlags);
-
-		PerShapeData->SetQueryEnabled(true);
-		PerShapeData->SetSimEnabled(true);
-		PerShapeData->SetIsProbe(false);
-		PerShapeData->SetShapeFilterData(ShapeFilter);
-		PerShapeData->SetCollisionTraceType(ConvertCollisionTraceFlag(ECollisionTraceFlag::CTF_UseSimpleAsComplex));
+		PerShapeData->SetQueryEnabled(InSetup.bEnableQuery);
+		PerShapeData->SetSimEnabled(InSetup.bEnableSim);
+		PerShapeData->SetIsProbe(InSetup.bEnableProbe);
+		PerShapeData->SetShapeFilterData(InSetup.ShapeFilterData);
+		PerShapeData->SetFilterInstanceData(InSetup.FilterInstanceData);
+		PerShapeData->SetCollisionTraceType(ConvertCollisionTraceFlag(InSetup.CollisionTraceType));
 		// TODO: UpdateShapeBounds?
-		// TODO: UserData
+		PerShapeData->SetUserData(InSetup.UserData);
 		// TODO: Materials
 
 		return PerShapeData;
@@ -222,14 +221,74 @@ namespace Chaos::Rigids::Async
 		return FTransform3f();
 	}
 
+	bool FRigidShapeInstanceAsync::GetQueryEnabled() const
+	{
+		return PerShapeData->GetQueryEnabled();
+	}
+
+	void FRigidShapeInstanceAsync::SetQueryEnabled(const bool bInEnabled)
+	{
+		PerShapeData->SetQueryEnabled(bInEnabled);
+	}
+
+	bool FRigidShapeInstanceAsync::GetSimEnabled() const
+	{
+		return PerShapeData->GetSimEnabled();
+	}
+
+	void FRigidShapeInstanceAsync::SetSimEnabled(const bool bInEnabled)
+	{
+		PerShapeData->SetSimEnabled(bInEnabled);
+	}
+
+	bool FRigidShapeInstanceAsync::GetIsProbe() const
+	{
+		return PerShapeData->GetIsProbe();
+	}
+
+	void FRigidShapeInstanceAsync::SetIsProbe(const bool bInIsProbe)
+	{
+		PerShapeData->SetIsProbe(bInIsProbe);
+	}
+
+	ECollisionTraceFlag FRigidShapeInstanceAsync::GetCollisionTraceType() const
+	{
+		return ConvertCollisionTraceFlag(PerShapeData->GetCollisionTraceType());
+	}
+
+	void FRigidShapeInstanceAsync::SetCollisionTraceType(const ECollisionTraceFlag InTraceFlag)
+	{
+		PerShapeData->SetCollisionTraceType(ConvertCollisionTraceFlag(InTraceFlag));
+	}
+
 	Chaos::Filter::FShapeFilterData FRigidShapeInstanceAsync::GetShapeFilter() const
 	{
 		return PerShapeData->GetShapeFilterData();
 	}
 
-	void FRigidShapeInstanceAsync::SetShapeFilter(const FShapeFilterData& ShapeFilter)
+	void FRigidShapeInstanceAsync::SetShapeFilter(const FShapeFilterData& InShapeFilter)
+	{
+		PerShapeData->SetShapeFilterData(InShapeFilter);
+	}
+
+	Chaos::Filter::FInstanceData FRigidShapeInstanceAsync::GetFilterInstanceData() const
+	{
+		return PerShapeData->GetFilterInstanceData();
+	}
+
+	void FRigidShapeInstanceAsync::SetFilterInstanceData(const Chaos::Filter::FInstanceData& InInstanceData)
+	{
+		PerShapeData->SetFilterInstanceData(InInstanceData);
+	}
+
+	void* FRigidShapeInstanceAsync::GetUserData() const
+	{
+		return PerShapeData->GetUserData();
+	}
+
+	void FRigidShapeInstanceAsync::SetUserData(void* InUserData)
 	{
-		PerShapeData->SetShapeFilterData(ShapeFilter);
+		PerShapeData->SetUserData(InUserData);
 	}
 } // namespace Chaos::Rigids::Async
 
diff --git a/Engine/Plugins/Experimental/ChaosRigidPhysicsAsync/Source/ChaosRigidPhysicsAsync/Public/ChaosRigidPhysicsAsync/RigidShapeInstanceAsync.h b/Engine/Plugins/Experimental/ChaosRigidPhysicsAsync/Source/ChaosRigidPhysicsAsync/Public/ChaosRigidPhysicsAsync/RigidShapeInstanceAsync.h
--- a/Engine/Plugins/Experimental/ChaosRigidPhysicsAsync/Source/ChaosRigidPhysicsAsync/Public/ChaosRigidPhysicsAsync/RigidShapeInstanceAsync.h
+++ b/Engine/Plugins/Experimental/ChaosRigidPhysicsAsync/Source/ChaosRigidPhysicsAsync/Public/ChaosRigidPhysicsAsync/RigidShapeInstanceAsync.h
@@ -26,8 +26,27 @@ namespace Chaos::Rigids::Async
 		UE_INTERNAL virtual UE::Physics::FAnyGeometry GetGeometry() const override final;
 
 		UE_INTERNAL virtual FTransform3f GetLocalTransform() const override final;
+
+		UE_INTERNAL virtual bool GetQueryEnabled() const override final;
+		UE_INTERNAL virtual void SetQueryEnabled(const bool bInEnabled) override final;
+
+		UE_INTERNAL virtual bool GetSimEnabled() const override final;
+		UE_INTERNAL virtual void SetSimEnabled(const bool bInEnabled) override final;
+
+		UE_INTERNAL virtual bool GetIsProbe() const override final;
+		UE_INTERNAL virtual void SetIsProbe(const bool bInIsProbe) override final;
+
+		UE_INTERNAL virtual ECollisionTraceFlag GetCollisionTraceType() const override final;
+		UE_INTERNAL virtual void SetCollisionTraceType(const ECollisionTraceFlag InTraceFlag) override final;
+
 		UE_INTERNAL virtual FShapeFilterData GetShapeFilter() const override final;
-		UE_INTERNAL virtual void SetShapeFilter(const FShapeFilterData& ShapeFilter) override final;
+		UE_INTERNAL virtual void SetShapeFilter(const FShapeFilterData& InShapeFilter) override final;
+
+		UE_INTERNAL virtual Chaos::Filter::FInstanceData GetFilterInstanceData() const override final;
+		UE_INTERNAL virtual void SetFilterInstanceData(const Chaos::Filter::FInstanceData& InInstanceData) override final;
+
+		UE_INTERNAL virtual void* GetUserData() const override final;
+		UE_INTERNAL virtual void SetUserData(void* InUserData) override final;
 
 		UE_INTERNAL Chaos::FPerShapeData* PerShapeData = nullptr;
 		UE_INTERNAL UE::Physics::FRigidShapeInstanceHandle Handle;
diff --git a/Engine/Plugins/Experimental/ChaosRigidPhysicsAsync/Tests/Private/ShapeInstanceTests.cpp b/Engine/Plugins/Experimental/ChaosRigidPhysicsAsync/Tests/Private/ShapeInstanceTests.cpp
--- a/Engine/Plugins/Experimental/ChaosRigidPhysicsAsync/Tests/Private/ShapeInstanceTests.cpp
+++ b/Engine/Plugins/Experimental/ChaosRigidPhysicsAsync/Tests/Private/ShapeInstanceTests.cpp
@@ -8,6 +8,7 @@
 #include "RigidPhysics/RigidScene.h"
 #include "RigidPhysics/RigidShapeInstance.h"
 #include "RigidTestFixture.h"
+#include <catch2/generators/catch_generators_range.hpp>
 
 #if UE_RIGIDPHYSICS_API_ENABLED
 
@@ -63,17 +64,32 @@ namespace Chaos::LowLevelTest
 		CHECK(AngularDistanceRad < AngularErrorThreshold);
 	}
 
-	void CreateAndCheckShapeInstanceAndTransform(FRigidTestFixture& Fixture, const FRigidShapeInstanceSetup& InSetup, const FAnyGeometry& ExpectedGeometry, const FTransform3f& ExpectedTransform)
+	struct FShapeInstanceTestFixture : public FRigidTestFixture
 	{
-		if (FRigidContextGameRW Context = Fixture.LockRWChecked())
+		FRigidShapeInstanceHandle CreateBodyAndShapeInstance(const FRigidShapeInstanceSetup& InSetup)
 		{
-			TRigidBodyPtr<FRigidContextGameRW> BodyPtr = Context->CreateBody(FRigidDebugName("Body1"), ERigidMovementType::Dynamic);
-			Fixture.AutoCleanup(BodyPtr);
+			FRigidShapeInstanceHandle ShapeInstanceHandle;
+			if (FRigidContextGameRW Context = LockRWChecked())
+			{
+				TRigidBodyPtr<FRigidContextGameRW> BodyPtr = Context->CreateBody(FRigidDebugName("Body"), ERigidMovementType::Dynamic);
+				AutoCleanup(BodyPtr);
 
-			TRigidShapeInstancePtr<FRigidContextGameRW> ShapeInstance = BodyPtr->CreateShape(InSetup);
+				ShapeInstanceHandle = BodyPtr->CreateShape(InSetup);
+				BodyPtr->Activate();
+			}
+			return ShapeInstanceHandle;
+		}
+	};
 
-			const FTransform3f ActualTransform = ShapeInstance->GetLocalTransform();
-			const FAnyGeometry ActualGeometry = ShapeInstance->GetGeometry();
+	void CreateAndCheckShapeInstanceAndTransform(FShapeInstanceTestFixture& Fixture, const FRigidShapeInstanceSetup& InSetup, const FAnyGeometry& ExpectedGeometry, const FTransform3f& ExpectedTransform)
+	{
+		FRigidShapeInstanceHandle ShapeInstanceHandle = Fixture.CreateBodyAndShapeInstance(InSetup);
+		if (FRigidContextGameRO Context = Fixture.LockROChecked())
+		{
+			const TRigidShapeInstancePtr<FRigidContextGameRO> ShapeInstancePtr = ShapeInstanceHandle.Pin(Context);
+
+			const FTransform3f ActualTransform = ShapeInstancePtr->GetLocalTransform();
+			const FAnyGeometry ActualGeometry = ShapeInstancePtr->GetGeometry();
 
 			ActualGeometry.Visit([&ExpectedTransform, &ActualTransform, &ExpectedGeometry]<typename GeometryType>(const GeometryType& ActualTypedGeometry)
 			{
@@ -87,14 +103,14 @@ namespace Chaos::LowLevelTest
 		}
 	}
 
-	void CreateAndCheckShapeInstanceAndTransform(FRigidTestFixture& Fixture, const FRigidShapeInstanceSetup& InSetup)
+	void CreateAndCheckShapeInstanceAndTransform(FShapeInstanceTestFixture& Fixture, const FRigidShapeInstanceSetup& InSetup)
 	{
 		CreateAndCheckShapeInstanceAndTransform(Fixture, InSetup, InSetup.Geometry, InSetup.LocalTransform);
 	}
 
 	TEST_CASE("ShapeInstance: Setup + GetShape/Transform", "[Chaos][API][ShapeInstance][unit]")
 	{
-		FRigidTestFixture Fixture;
+		FShapeInstanceTestFixture Fixture;
 		SECTION("Box")
 		{
 			FRigidShapeInstanceSetup Setup(FBoxGeometry(FVector3f(1, 2, 3)));
@@ -299,6 +315,274 @@ namespace Chaos::LowLevelTest
 			}
 		}
 	}
+
+	TEST_CASE("ShapeInstance: QueryEnabled", "[Chaos][API][ShapeInstance][unit]")
+	{
+		FShapeInstanceTestFixture Fixture;
+		FRigidShapeInstanceSetup Setup(FSphereGeometry(1));
+
+		const bool bQueryEnabled = GENERATE(true, false);
+		Setup.bEnableQuery = bQueryEnabled;
+
+		FRigidShapeInstanceHandle ShapeInstanceHandle = Fixture.CreateBodyAndShapeInstance(Setup);
+		if (FRigidContextGameRO Context = Fixture.LockROChecked())
+		{
+			const TRigidShapeInstancePtr<FRigidContextGameRO> ShapeInstancePtr = ShapeInstanceHandle.Pin(Context);
+			CHECK(ShapeInstancePtr->GetQueryEnabled() == bQueryEnabled);
+		}
+
+		Fixture.RunPTCallback([ShapeInstanceHandle, bQueryEnabled](const FRigidContextSimRW& Context)
+			{
+				const TRigidShapeInstancePtr<FRigidContextSimRW> ShapeInstancePtr = ShapeInstanceHandle.Pin(Context);
+				CHECK(ShapeInstancePtr->GetQueryEnabled() == bQueryEnabled);
+			});
+
+		const bool bNewQueryEnabled = !bQueryEnabled;
+		if (FRigidContextGameRW Context = Fixture.LockRWChecked())
+		{
+			TRigidShapeInstancePtr<FRigidContextGameRW> ShapeInstancePtr = ShapeInstanceHandle.Pin(Context);
+			ShapeInstancePtr->SetQueryEnabled(bNewQueryEnabled);
+			CHECK(ShapeInstancePtr->GetQueryEnabled() == bNewQueryEnabled);
+		}
+
+		Fixture.RunPTCallback([ShapeInstanceHandle, bNewQueryEnabled](const FRigidContextSimRW& Context)
+			{
+				const TRigidShapeInstancePtr<FRigidContextSimRW> ShapeInstancePtr = ShapeInstanceHandle.Pin(Context);
+				CHECK(ShapeInstancePtr->GetQueryEnabled() == bNewQueryEnabled);
+			});
+	}
+
+	TEST_CASE("ShapeInstance: SimEnabled", "[Chaos][API][ShapeInstance][unit]")
+	{
+		FShapeInstanceTestFixture Fixture;
+		FRigidShapeInstanceSetup Setup(FSphereGeometry(1));
+
+		const bool bSimEnabled = GENERATE(true, false);
+		Setup.bEnableSim = bSimEnabled;
+
+		FRigidShapeInstanceHandle ShapeInstanceHandle = Fixture.CreateBodyAndShapeInstance(Setup);
+		if (FRigidContextGameRO Context = Fixture.LockROChecked())
+		{
+			const TRigidShapeInstancePtr<FRigidContextGameRO> ShapeInstancePtr = ShapeInstanceHandle.Pin(Context);
+			CHECK(ShapeInstancePtr->GetSimEnabled() == bSimEnabled);
+		}
+
+		Fixture.RunPTCallback([ShapeInstanceHandle, bSimEnabled](const FRigidContextSimRW& Context)
+			{
+				const TRigidShapeInstancePtr<FRigidContextSimRW> ShapeInstancePtr = ShapeInstanceHandle.Pin(Context);
+				CHECK(ShapeInstancePtr->GetSimEnabled() == bSimEnabled);
+			});
+
+		const bool bNewSimEnabled = !bSimEnabled;
+		if (FRigidContextGameRW Context = Fixture.LockRWChecked())
+		{
+			TRigidShapeInstancePtr<FRigidContextGameRW> ShapeInstancePtr = ShapeInstanceHandle.Pin(Context);
+			ShapeInstancePtr->SetSimEnabled(bNewSimEnabled);
+			CHECK(ShapeInstancePtr->GetSimEnabled() == bNewSimEnabled);
+		}
+
+		Fixture.RunPTCallback([ShapeInstanceHandle, bNewSimEnabled](const FRigidContextSimRW& Context)
+			{
+				const TRigidShapeInstancePtr<FRigidContextSimRW> ShapeInstancePtr = ShapeInstanceHandle.Pin(Context);
+				CHECK(ShapeInstancePtr->GetSimEnabled() == bNewSimEnabled);
+			});
+	}
+
+	TEST_CASE("ShapeInstance: IsProbe", "[Chaos][API][ShapeInstance][unit]")
+	{
+		FShapeInstanceTestFixture Fixture;
+		FRigidShapeInstanceSetup Setup(FSphereGeometry(1));
+
+		const bool bIsProbe = GENERATE(false, true);
+		Setup.bEnableProbe = bIsProbe;
+
+		FRigidShapeInstanceHandle ShapeInstanceHandle = Fixture.CreateBodyAndShapeInstance(Setup);
+		if (FRigidContextGameRO Context = Fixture.LockROChecked())
+		{
+			const TRigidShapeInstancePtr<FRigidContextGameRO> ShapeInstancePtr = ShapeInstanceHandle.Pin(Context);
+			CHECK(ShapeInstancePtr->GetIsProbe() == bIsProbe);
+		}
+
+		Fixture.RunPTCallback([ShapeInstanceHandle, bIsProbe](const FRigidContextSimRW& Context)
+			{
+				const TRigidShapeInstancePtr<FRigidContextSimRW> ShapeInstancePtr = ShapeInstanceHandle.Pin(Context);
+				CHECK(ShapeInstancePtr->GetIsProbe() == bIsProbe);
+			});
+
+		const bool bNewIsProbe = !bIsProbe;
+		if (FRigidContextGameRW Context = Fixture.LockRWChecked())
+		{
+			TRigidShapeInstancePtr<FRigidContextGameRW> ShapeInstancePtr = ShapeInstanceHandle.Pin(Context);
+			ShapeInstancePtr->SetIsProbe(bNewIsProbe);
+			CHECK(ShapeInstancePtr->GetIsProbe() == bNewIsProbe);
+		}
+
+		Fixture.RunPTCallback([ShapeInstanceHandle, bNewIsProbe](const FRigidContextSimRW& Context)
+			{
+				const TRigidShapeInstancePtr<FRigidContextSimRW> ShapeInstancePtr = ShapeInstanceHandle.Pin(Context);
+				CHECK(ShapeInstancePtr->GetIsProbe() == bNewIsProbe);
+			});
+	}
+
+	TEST_CASE("ShapeInstance: CollisionTraceType", "[Chaos][API][ShapeInstance][unit]")
+	{
+		FShapeInstanceTestFixture Fixture;
+		const TArray<ECollisionTraceFlag> TraceTypes
+		{
+			ECollisionTraceFlag::CTF_UseDefault,
+			ECollisionTraceFlag::CTF_UseSimpleAndComplex,
+			ECollisionTraceFlag::CTF_UseSimpleAsComplex,
+			ECollisionTraceFlag::CTF_UseComplexAsSimple,
+		};
+		const uint32 Index = GENERATE_COPY(Catch::Generators::range(0, TraceTypes.Num()));
+		const ECollisionTraceFlag TraceType = TraceTypes[Index];
+		const ECollisionTraceFlag NewTraceType = TraceTypes[(Index + 1) % TraceTypes.Num()];
+
+		FRigidShapeInstanceSetup Setup(FSphereGeometry(1));
+		Setup.CollisionTraceType = TraceType;
+
+		FRigidShapeInstanceHandle ShapeInstanceHandle = Fixture.CreateBodyAndShapeInstance(Setup);
+		if (FRigidContextGameRO Context = Fixture.LockROChecked())
+		{
+			const TRigidShapeInstancePtr<FRigidContextGameRO> ShapeInstancePtr = ShapeInstanceHandle.Pin(Context);
+			CHECK(ShapeInstancePtr->GetCollisionTraceType() == TraceType);
+		}
+
+		Fixture.RunPTCallback([ShapeInstanceHandle, TraceType](const FRigidContextSimRW& Context)
+			{
+				const TRigidShapeInstancePtr<FRigidContextSimRW> ShapeInstancePtr = ShapeInstanceHandle.Pin(Context);
+				CHECK(ShapeInstancePtr->GetCollisionTraceType() == TraceType);
+			});
+
+		if (FRigidContextGameRW Context = Fixture.LockRWChecked())
+		{
+			TRigidShapeInstancePtr<FRigidContextGameRW> ShapeInstancePtr = ShapeInstanceHandle.Pin(Context);
+			ShapeInstancePtr->SetCollisionTraceType(NewTraceType);
+			CHECK(ShapeInstancePtr->GetCollisionTraceType() == NewTraceType);
+		}
+
+		Fixture.RunPTCallback([ShapeInstanceHandle, NewTraceType](const FRigidContextSimRW& Context)
+			{
+				const TRigidShapeInstancePtr<FRigidContextSimRW> ShapeInstancePtr = ShapeInstanceHandle.Pin(Context);
+				CHECK(ShapeInstancePtr->GetCollisionTraceType() == NewTraceType);
+			});
+	}
+
+	TEST_CASE("ShapeInstance: ShapeFilter", "[Chaos][API][ShapeInstance][unit]")
+	{
+		FShapeInstanceTestFixture Fixture;
+
+		const TStaticArray<Chaos::Filter::FShapeFilterData, 2> ShapeFilters
+		{
+			Chaos::Filter::FShapeFilterBuilder().SetCollisionChannelIndex(1).Build(),
+			Chaos::Filter::FShapeFilterBuilder().SetCollisionChannelIndex(2).Build(),
+		};
+
+		FRigidShapeInstanceSetup Setup(FSphereGeometry(1));
+		Setup.ShapeFilterData = ShapeFilters[0];
+
+		FRigidShapeInstanceHandle ShapeInstanceHandle = Fixture.CreateBodyAndShapeInstance(Setup);
+		if (FRigidContextGameRO Context = Fixture.LockROChecked())
+		{
+			const TRigidShapeInstancePtr<FRigidContextGameRO> ShapeInstancePtr = ShapeInstanceHandle.Pin(Context);
+			CHECK(ShapeInstancePtr->GetShapeFilter() == ShapeFilters[0]);
+		}
+
+		Fixture.RunPTCallback([ShapeInstanceHandle, &ShapeFilters](const FRigidContextSimRW& Context)
+			{
+				const TRigidShapeInstancePtr<FRigidContextSimRW> ShapeInstancePtr = ShapeInstanceHandle.Pin(Context);
+				CHECK(ShapeInstancePtr->GetShapeFilter() == ShapeFilters[0]);
+			});
+
+		if (FRigidContextGameRW Context = Fixture.LockRWChecked())
+		{
+			TRigidShapeInstancePtr<FRigidContextGameRW> ShapeInstancePtr = ShapeInstanceHandle.Pin(Context);
+			ShapeInstancePtr->SetShapeFilter(ShapeFilters[1]);
+			CHECK(ShapeInstancePtr->GetShapeFilter() == ShapeFilters[1]);
+		}
+
+		Fixture.RunPTCallback([ShapeInstanceHandle, &ShapeFilters](const FRigidContextSimRW& Context)
+			{
+				const TRigidShapeInstancePtr<FRigidContextSimRW> ShapeInstancePtr = ShapeInstanceHandle.Pin(Context);
+				CHECK(ShapeInstancePtr->GetShapeFilter() == ShapeFilters[1]);
+			});
+	}
+
+	TEST_CASE("ShapeInstance: FilterInstanceData", "[Chaos][API][ShapeInstance][unit]")
+	{
+		FShapeInstanceTestFixture Fixture;
+
+		const TStaticArray<Chaos::Filter::FInstanceData, 2> InstanceData
+		{
+			Chaos::Filter::FInstanceData(123, 456),
+			Chaos::Filter::FInstanceData(789, 1011),
+		};
+
+		FRigidShapeInstanceSetup Setup(FSphereGeometry(1));
+		Setup.FilterInstanceData = InstanceData[0];
+
+		FRigidShapeInstanceHandle ShapeInstanceHandle = Fixture.CreateBodyAndShapeInstance(Setup);
+		if (FRigidContextGameRO Context = Fixture.LockROChecked())
+		{
+			const TRigidShapeInstancePtr<FRigidContextGameRO> ShapeInstancePtr = ShapeInstanceHandle.Pin(Context);
+			CHECK(ShapeInstancePtr->GetFilterInstanceData() == InstanceData[0]);
+		}
+
+		Fixture.RunPTCallback([ShapeInstanceHandle, &InstanceData](const FRigidContextSimRW& Context)
+			{
+				const TRigidShapeInstancePtr<FRigidContextSimRW> ShapeInstancePtr = ShapeInstanceHandle.Pin(Context);
+				CHECK(ShapeInstancePtr->GetFilterInstanceData() == InstanceData[0]);
+			});
+
+		if (FRigidContextGameRW Context = Fixture.LockRWChecked())
+		{
+			TRigidShapeInstancePtr<FRigidContextGameRW> ShapeInstancePtr = ShapeInstanceHandle.Pin(Context);
+			ShapeInstancePtr->SetFilterInstanceData(InstanceData[1]);
+			CHECK(ShapeInstancePtr->GetFilterInstanceData() == InstanceData[1]);
+		}
+
+		Fixture.RunPTCallback([ShapeInstanceHandle, &InstanceData](const FRigidContextSimRW& Context)
+			{
+				const TRigidShapeInstancePtr<FRigidContextSimRW> ShapeInstancePtr = ShapeInstanceHandle.Pin(Context);
+				CHECK(ShapeInstancePtr->GetFilterInstanceData() == InstanceData[1]);
+			});
+	}
+
+	TEST_CASE("ShapeInstance: UserData", "[Chaos][API][ShapeInstance][unit]")
+	{
+		FShapeInstanceTestFixture Fixture;
+
+		TStaticArray<int32, 2> UserData{ 1, 2 };
+
+		FRigidShapeInstanceSetup Setup(FSphereGeometry(1));
+		Setup.UserData = &UserData[0];
+
+		FRigidShapeInstanceHandle ShapeInstanceHandle = Fixture.CreateBodyAndShapeInstance(Setup);
+		if (FRigidContextGameRO Context = Fixture.LockROChecked())
+		{
+			const TRigidShapeInstancePtr<FRigidContextGameRO> ShapeInstancePtr = ShapeInstanceHandle.Pin(Context);
+			CHECK(ShapeInstancePtr->GetUserData() == &UserData[0]);
+		}
+
+		Fixture.RunPTCallback([ShapeInstanceHandle, &UserData](const FRigidContextSimRW& Context)
+			{
+				const TRigidShapeInstancePtr<FRigidContextSimRW> ShapeInstancePtr = ShapeInstanceHandle.Pin(Context);
+				CHECK(ShapeInstancePtr->GetUserData() == &UserData[0]);
+			});
+
+		if (FRigidContextGameRW Context = Fixture.LockRWChecked())
+		{
+			TRigidShapeInstancePtr<FRigidContextGameRW> ShapeInstancePtr = ShapeInstanceHandle.Pin(Context);
+			ShapeInstancePtr->SetUserData(&UserData[1]);
+			CHECK(ShapeInstancePtr->GetUserData() == &UserData[1]);
+		}
+
+		Fixture.RunPTCallback([ShapeInstanceHandle, &UserData](const FRigidContextSimRW& Context)
+			{
+				const TRigidShapeInstancePtr<FRigidContextSimRW> ShapeInstancePtr = ShapeInstanceHandle.Pin(Context);
+				CHECK(ShapeInstancePtr->GetUserData() == &UserData[1]);
+			});
+	}
 } // namespace Chaos::LowLevelTest
 
 #endif // UE_RIGIDPHYSICS_API_ENABLED
diff --git a/Engine/Source/Runtime/Experimental/RigidPhysics/Private/RigidPhysics/RigidShapeInstanceSetup.cpp b/Engine/Source/Runtime/Experimental/RigidPhysics/Private/RigidPhysics/RigidShapeInstanceSetup.cpp
new file mode 100644
--- /dev/null
+++ b/Engine/Source/Runtime/Experimental/RigidPhysics/Private/RigidPhysics/RigidShapeInstanceSetup.cpp
@@ -0,0 +1,17 @@
+// Copyright Epic Games, Inc. All Rights Reserved.
+#include "RigidPhysics/RigidShapeInstanceSetup.h"
+
+#if UE_RIGIDPHYSICS_API_ENABLED
+
+namespace UE::Physics
+{
+	static constexpr Chaos::EFilterFlags DefaultFilterFlags = Chaos::EFilterFlags::SimpleCollision | Chaos::EFilterFlags::ContactNotify | Chaos::EFilterFlags::ModifyContacts;
+
+	FRigidShapeInstanceSetup::FRigidShapeInstanceSetup(const FAnyGeometry& InGeometry)
+		: Geometry(InGeometry)
+		, ShapeFilterData(Chaos::Filter::FShapeFilterBuilder::BuildBlockAll(DefaultFilterFlags))
+	{
+	}
+} // namespace UE::Physics
+
+#endif // UE_RIGIDPHYSICS_API_ENABLED
diff --git a/Engine/Source/Runtime/Experimental/RigidPhysics/Public/RigidPhysics/Internal/IRigidShapeInstance.h b/Engine/Source/Runtime/Experimental/RigidPhysics/Public/RigidPhysics/Internal/IRigidShapeInstance.h
--- a/Engine/Source/Runtime/Experimental/RigidPhysics/Public/RigidPhysics/Internal/IRigidShapeInstance.h
+++ b/Engine/Source/Runtime/Experimental/RigidPhysics/Public/RigidPhysics/Internal/IRigidShapeInstance.h
@@ -5,6 +5,7 @@
 
 #if UE_RIGIDPHYSICS_API_ENABLED
 
+#include "BodySetupEnums.h"
 #include "Chaos/CollisionFilterData.h"
 #include "Math/Transform.h"
 #include "RigidPhysics/Geometry/AnyGeometry.h"
@@ -31,8 +32,26 @@ namespace UE::Physics
 		UE_INTERNAL virtual FTransform3f GetLocalTransform() const = 0;
 		// TODO: SetLocalTransform
 
+		UE_INTERNAL virtual bool GetQueryEnabled() const = 0;
+		UE_INTERNAL virtual void SetQueryEnabled(const bool bInEnabled) = 0;
+
+		UE_INTERNAL virtual bool GetSimEnabled() const = 0;
+		UE_INTERNAL virtual void SetSimEnabled(const bool bInEnabled) = 0;
+
+		UE_INTERNAL virtual bool GetIsProbe() const = 0;
+		UE_INTERNAL virtual void SetIsProbe(const bool bInIsProbe) = 0;
+
+		UE_INTERNAL virtual ECollisionTraceFlag GetCollisionTraceType() const = 0;
+		UE_INTERNAL virtual void SetCollisionTraceType(const ECollisionTraceFlag InTraceFlag) = 0;
+
 		UE_INTERNAL virtual Chaos::Filter::FShapeFilterData GetShapeFilter() const = 0;
-		UE_INTERNAL virtual void SetShapeFilter(const Chaos::Filter::FShapeFilterData& ShapeFilter) = 0;
+		UE_INTERNAL virtual void SetShapeFilter(const Chaos::Filter::FShapeFilterData& InShapeFilter) = 0;
+
+		UE_INTERNAL virtual Chaos::Filter::FInstanceData GetFilterInstanceData() const = 0;
+		UE_INTERNAL virtual void SetFilterInstanceData(const Chaos::Filter::FInstanceData& InInstanceData) = 0;
+
+		UE_INTERNAL virtual void* GetUserData() const = 0;
+		UE_INTERNAL virtual void SetUserData(void* InUserData) = 0;
 	};
 } // namespace UE::Physics
 
diff --git a/Engine/Source/Runtime/Experimental/RigidPhysics/Public/RigidPhysics/RigidShapeInstance.h b/Engine/Source/Runtime/Experimental/RigidPhysics/Public/RigidPhysics/RigidShapeInstance.h
--- a/Engine/Source/Runtime/Experimental/RigidPhysics/Public/RigidPhysics/RigidShapeInstance.h
+++ b/Engine/Source/Runtime/Experimental/RigidPhysics/Public/RigidPhysics/RigidShapeInstance.h
@@ -71,14 +71,74 @@ namespace UE::Physics
 		}
 		// TODO: SetLocalTransform
 
+		bool GetQueryEnabled() const
+		{
+			return ShapeInstanceRaw->GetQueryEnabled();
+		}
+
+		void SetQueryEnabled(const bool bInEnabled) requires CIsGameContext<FContext>
+		{
+			ShapeInstanceRaw->SetQueryEnabled(bInEnabled);
+		}
+
+		bool GetSimEnabled() const
+		{
+			return ShapeInstanceRaw->GetSimEnabled();
+		}
+
+		void SetSimEnabled(const bool bInEnabled) requires CIsGameContext<FContext>
+		{
+			ShapeInstanceRaw->SetSimEnabled(bInEnabled);
+		}
+
+		bool GetIsProbe() const
+		{
+			return ShapeInstanceRaw->GetIsProbe();
+		}
+
+		void SetIsProbe(const bool bInIsProbe) requires CIsGameContext<FContext>
+		{
+			ShapeInstanceRaw->SetIsProbe(bInIsProbe);
+		}
+
+		ECollisionTraceFlag GetCollisionTraceType() const
+		{
+			return ShapeInstanceRaw->GetCollisionTraceType();
+		}
+
+		void SetCollisionTraceType(const ECollisionTraceFlag InTraceFlag) requires CIsGameContext<FContext>
+		{
+			ShapeInstanceRaw->SetCollisionTraceType(InTraceFlag);
+		}
+
 		Chaos::Filter::FShapeFilterData GetShapeFilter() const
 		{
 			return ShapeInstanceRaw->GetShapeFilter();
 		}
 
-		void SetShapeFilter(const Chaos::Filter::FShapeFilterData& ShapeFilter)
+		void SetShapeFilter(const Chaos::Filter::FShapeFilterData& InShapeFilter) requires CIsGameContext<FContext>
+		{
+			ShapeInstanceRaw->SetShapeFilter(InShapeFilter);
+		}
+
+		Chaos::Filter::FInstanceData GetFilterInstanceData() const
+		{
+			return ShapeInstanceRaw->GetFilterInstanceData();
+		}
+
+		void SetFilterInstanceData(const Chaos::Filter::FInstanceData& InInstanceData) requires CIsGameContext<FContext>
+		{
+			ShapeInstanceRaw->SetFilterInstanceData(InInstanceData);
+		}
+
+		void* GetUserData() const
+		{
+			return ShapeInstanceRaw->GetUserData();
+		}
+
+		void SetUserData(void* InUserData) requires CIsGameContext<FContext>
 		{
-			ShapeInstanceRaw->SetShapeFilter(ShapeFilter);
+			ShapeInstanceRaw->SetUserData(InUserData);
 		}
 
 		UE_INTERNAL void Reset()
diff --git a/Engine/Source/Runtime/Experimental/RigidPhysics/Public/RigidPhysics/RigidShapeInstanceSetup.h b/Engine/Source/Runtime/Experimental/RigidPhysics/Public/RigidPhysics/RigidShapeInstanceSetup.h
--- a/Engine/Source/Runtime/Experimental/RigidPhysics/Public/RigidPhysics/RigidShapeInstanceSetup.h
+++ b/Engine/Source/Runtime/Experimental/RigidPhysics/Public/RigidPhysics/RigidShapeInstanceSetup.h
@@ -6,6 +6,7 @@
 #if UE_RIGIDPHYSICS_API_ENABLED
 #include "Math/Transform.h"
 #include "RigidPhysics/Geometry/AnyGeometry.h"
+#include "BodySetupEnums.h"
 
 UE_EXPERIMENTAL(5.8, "The new Chaos API is experimental")
 namespace UE::Physics
@@ -13,13 +14,17 @@ namespace UE::Physics
 	class FRigidShapeInstanceSetup
 	{
 	public:
-		FRigidShapeInstanceSetup(const FAnyGeometry& InGeometry)
-			: Geometry(InGeometry)
-		{
-		}
+		RIGIDPHYSICS_API FRigidShapeInstanceSetup(const FAnyGeometry& InGeometry);
 
 		FAnyGeometry Geometry;
 		FTransform3f LocalTransform;
+		Chaos::Filter::FShapeFilterData ShapeFilterData;
+		Chaos::Filter::FInstanceData FilterInstanceData;
+		ECollisionTraceFlag CollisionTraceType = ECollisionTraceFlag::CTF_UseDefault;
+		bool bEnableSim = true;
+		bool bEnableQuery = true;
+		bool bEnableProbe = false;
+		void* UserData = nullptr;
 	};
 } // namespace UE::Physics
 
