Adding chaos tests to low level tests for horde.

#Synced-CL 52100670

#Synced-CL 52100670

diff --git a/Engine/Build/LowLevelTests/ChaosPhysics.xml b/Engine/Build/LowLevelTests/ChaosPhysics.xml
new file mode 100644
--- /dev/null
+++ b/Engine/Build/LowLevelTests/ChaosPhysics.xml
@@ -0,0 +1,9 @@
+﻿<?xml version="1.0" encoding="utf-8"?>
+<BuildGraph xmlns:xsi="http://www.w3.org/2001/XMLSchema-instance" xsi:schemaLocation="http://www.epicgames.com/BuildGraph ../../Build/Graph/Schema.xsd" xmlns="http://www.epicgames.com/BuildGraph">
+  <Option Name="RunChaosPhysicsTests" DefaultValue="" Description="Run ChaosPhysics Tests" />
+  <Property Name="TestNames" Value="$(TestNames);ChaosPhysics" />
+  <Extend Name="RunAllTests">
+    <Expand Name="DeployAndTest" Platform="Win64" TestName="ChaosPhysics" ShortName="ChaosPhysics" TargetName="ChaosPhysicsTests" BinaryRelativePath="Engine\Binaries" ReportType="testdata" Tags="[@ChaosTests]" ExtraCompilationArgs="-allmodules" />
+    <Expand Name="DeployAndTest" Platform="Mac" TestName="ChaosPhysics" ShortName="ChaosPhysics" TargetName="ChaosPhysicsTests" BinaryRelativePath="Engine\Binaries" ReportType="testdata" Tags="[@ChaosTests]" ExtraCompilationArgs="-allmodules" />
+  </Extend>
+</BuildGraph>
\ No newline at end of file
diff --git a/Engine/Build/LowLevelTests/ChaosRigidPhysicsAsync.xml b/Engine/Build/LowLevelTests/ChaosRigidPhysicsAsync.xml
new file mode 100644
--- /dev/null
+++ b/Engine/Build/LowLevelTests/ChaosRigidPhysicsAsync.xml
@@ -0,0 +1,9 @@
+﻿<?xml version="1.0" encoding="utf-8"?>
+<BuildGraph xmlns:xsi="http://www.w3.org/2001/XMLSchema-instance" xsi:schemaLocation="http://www.epicgames.com/BuildGraph ../../Build/Graph/Schema.xsd" xmlns="http://www.epicgames.com/BuildGraph">
+  <Option Name="RunChaosRigidPhysicsAsyncTests" DefaultValue="" Description="Run ChaosRigidPhysicsAsync Tests" />
+  <Property Name="TestNames" Value="$(TestNames);ChaosRigidPhysicsAsync" />
+  <Extend Name="RunAllTests">
+    <Expand Name="DeployAndTest" Platform="Win64" TestName="ChaosRigidPhysicsAsync" ShortName="ChaosRigidPhysicsAsync" TargetName="ChaosRigidPhysicsAsyncTests" BinaryRelativePath="Engine\Binaries" ReportType="testdata" Tags="[@ChaosTests]" ExtraCompilationArgs="-allmodules" />
+    <Expand Name="DeployAndTest" Platform="Mac" TestName="ChaosRigidPhysicsAsync" ShortName="ChaosRigidPhysicsAsync" TargetName="ChaosRigidPhysicsAsyncTests" BinaryRelativePath="Engine\Binaries" ReportType="testdata" Tags="[@ChaosTests]" ExtraCompilationArgs="-allmodules" />
+  </Extend>
+</BuildGraph>
\ No newline at end of file
diff --git a/Engine/Plugins/Experimental/ChaosRigidPhysicsAsync/Tests/ChaosRigidPhysicsAsyncTests.Build.cs b/Engine/Plugins/Experimental/ChaosRigidPhysicsAsync/Tests/ChaosRigidPhysicsAsyncTests.Build.cs
--- a/Engine/Plugins/Experimental/ChaosRigidPhysicsAsync/Tests/ChaosRigidPhysicsAsyncTests.Build.cs
+++ b/Engine/Plugins/Experimental/ChaosRigidPhysicsAsync/Tests/ChaosRigidPhysicsAsyncTests.Build.cs
@@ -9,13 +9,14 @@ public class ChaosRigidPhysicsAsyncTests : TestModuleRules
 		if (InTestMode)
 		{
 			TestMetadata = new Metadata();
-			TestMetadata.TestName = "ChaosRigidPhysicsAsyncTests";
-			TestMetadata.TestShortName = "ChaosRigidPhysicsAsyncTests";
+			TestMetadata.TestName = "ChaosRigidPhysicsAsync";
+			TestMetadata.TestShortName = "ChaosRigidPhysicsAsync";
 			TestMetadata.ReportType = "testdata";
-			TestMetadata.SupportedPlatforms.Add(UnrealTargetPlatform.Linux);
 			TestMetadata.SupportedPlatforms.Add(UnrealTargetPlatform.Mac);
-			TestMetadata.SupportedPlatforms.Add(UnrealTargetPlatform.Android);
-			TestMetadata.SupportedPlatforms.Add(UnrealTargetPlatform.IOS);
+			// TODO: These don't currently work. Need to investigate later, but disabling to get something working for now.
+			//TestMetadata.SupportedPlatforms.Add(UnrealTargetPlatform.Linux);
+			//TestMetadata.SupportedPlatforms.Add(UnrealTargetPlatform.Android);
+			//TestMetadata.SupportedPlatforms.Add(UnrealTargetPlatform.IOS);
 
 			string PlatformCompilationArgs;
 			foreach (var Platform in UnrealTargetPlatform.GetValidPlatforms())
@@ -29,11 +30,13 @@ public class ChaosRigidPhysicsAsyncTests : TestModuleRules
 					PlatformCompilationArgs = "-allmodules";
 				}
 				TestMetadata.PlatformCompilationExtraArgs.Add(Platform, PlatformCompilationArgs);
+				// Add our test tag to limit tests
+				TestMetadata.PlatformTags.Add(Platform, "[@ChaosTests]");
 			}
 
 			// Platform-specific tags
-			TestMetadata.PlatformTags.Add(UnrealTargetPlatform.Linux, "~[.]~[Slow]");
-			TestMetadata.PlatformTags.Add(UnrealTargetPlatform.Android, "~[Perf]~[Slow]~[AndroidSkip]");
+			TestMetadata.PlatformTags[UnrealTargetPlatform.Linux] += "~[.]~[Slow]";
+			TestMetadata.PlatformTags[UnrealTargetPlatform.Android] += "~[Perf]~[Slow]~[AndroidSkip]";
 
 			// Allow Android run for this test
 			// Will remove Android from PlatformsRunUnsupported as more diverse types of tests can run on this platform
diff --git a/Engine/Source/Programs/ChaosPhysicsTests/ChaosPhysicsTests.Build.cs b/Engine/Source/Programs/ChaosPhysicsTests/ChaosPhysicsTests.Build.cs
--- a/Engine/Source/Programs/ChaosPhysicsTests/ChaosPhysicsTests.Build.cs
+++ b/Engine/Source/Programs/ChaosPhysicsTests/ChaosPhysicsTests.Build.cs
@@ -9,13 +9,14 @@ public class ChaosPhysicsTests : TestModuleRules
 		if (InTestMode)
 		{
 			TestMetadata = new Metadata();
-			TestMetadata.TestName = "ChaosPhysicsTests";
-			TestMetadata.TestShortName = "ChaosPhysicsTests";
+			TestMetadata.TestName = "ChaosPhysics";
+			TestMetadata.TestShortName = "ChaosPhysics";
 			TestMetadata.ReportType = "testdata";
-			TestMetadata.SupportedPlatforms.Add(UnrealTargetPlatform.Linux);
 			TestMetadata.SupportedPlatforms.Add(UnrealTargetPlatform.Mac);
-			TestMetadata.SupportedPlatforms.Add(UnrealTargetPlatform.Android);
-			TestMetadata.SupportedPlatforms.Add(UnrealTargetPlatform.IOS);
+			// TODO: These don't currently work. Need to investigate later, but disabling to get something working for now.
+			//TestMetadata.SupportedPlatforms.Add(UnrealTargetPlatform.Linux);
+			//TestMetadata.SupportedPlatforms.Add(UnrealTargetPlatform.Android);
+			//TestMetadata.SupportedPlatforms.Add(UnrealTargetPlatform.IOS);
 
 			string PlatformCompilationArgs;
 			foreach (var Platform in UnrealTargetPlatform.GetValidPlatforms())
@@ -29,11 +30,13 @@ public class ChaosPhysicsTests : TestModuleRules
 					PlatformCompilationArgs = "-allmodules";
 				}
 				TestMetadata.PlatformCompilationExtraArgs.Add(Platform, PlatformCompilationArgs);
+				// Add our test tag to limit tests
+				TestMetadata.PlatformTags.Add(Platform, "[@ChaosTests]");
 			}
 
 			// Platform-specific tags
-			TestMetadata.PlatformTags.Add(UnrealTargetPlatform.Linux, "~[.]~[Slow]");
-			TestMetadata.PlatformTags.Add(UnrealTargetPlatform.Android, "~[Perf]~[Slow]~[AndroidSkip]");
+			TestMetadata.PlatformTags[UnrealTargetPlatform.Linux] += "~[.]~[Slow]";
+			TestMetadata.PlatformTags[UnrealTargetPlatform.Android] += "~[Perf]~[Slow]~[AndroidSkip]";
 
 			// Allow Android run for this test
 			// Will remove Android from PlatformsRunUnsupported as more diverse types of tests can run on this platform
