Add FBindableTransform for binding to FTransform values (it has less overhead than using FBindableStruct with a FTransform)
Fix a bug where FBindableEnum would not allow you to set enum values which correspond to 0

#Synced-CL 52179140

diff --git a/Engine/Plugins/Experimental/UAF/UAF/Source/UAF/Private/BindableValue/UAFBindableTypes.cpp b/Engine/Plugins/Experimental/UAF/UAF/Source/UAF/Private/BindableValue/UAFBindableTypes.cpp
--- a/Engine/Plugins/Experimental/UAF/UAF/Source/UAF/Private/BindableValue/UAFBindableTypes.cpp
+++ b/Engine/Plugins/Experimental/UAF/UAF/Source/UAF/Private/BindableValue/UAFBindableTypes.cpp
@@ -258,6 +258,11 @@ FQuat FBindableQuat::GetValue(FUAFAssetInstance* Instance) const
 	return UE::UAF::Private::ResolveBindableValue<FQuat>(*this, Instance);
 }
 
+FTransform FBindableTransform::GetValue(FUAFAssetInstance* Instance) const
+{
+	return UE::UAF::Private::ResolveBindableValue<FTransform>(*this, Instance);
+}
+
 //------------------------------------------------------------------------------
 // FBindableEnum
 //------------------------------------------------------------------------------
@@ -633,14 +638,16 @@ bool FBindableValueBase::ImportBindingText(const TCHAR*& Buffer, TUniquePtr<FUAF
 bool FBindableValueBase::ExportTextItem(FString& ValueStr, const FBindableValueBase& DefaultValue,
 	UObject* Parent, int32 PortFlags, UObject* ExportRootScope) const
 {
-	// Always compare against a proper type default rather than the passed-in DefaultValue.
-	// This avoids the self-comparison issue when callers pass Data as both value and default
-	// (e.g. FRigVMMemoryStorageStruct::GetDataAsString), which would cause all UPROPERTYs to
-	// appear unchanged and be omitted from the exported text.
-	const FBindableValueBase TypeDefault;
+	// Pass `this` as both Value and Defaults to force ALL UPROPERTYs to be exported.
+	// FProperty::ExportText_Direct unconditionally exports when Data==Delta (pointer equality).
+	// This avoids two issues:
+	// 1. The self-comparison issue when callers pass Data as both value and default
+	//    (e.g. FRigVMMemoryStorageStruct::GetDataAsString).
+	// 2. Properties whose value happens to match a default-constructed instance being omitted
+	//    (e.g. FBindableEnum::ConstantValue == 0 matching the default int32 value).
 	FBindableValueBase::StaticStruct()->ExportText(ValueStr,
 		this,
-		&TypeDefault,
+		this,
 		Parent,
 		PortFlags,
 		ExportRootScope,
@@ -674,8 +681,7 @@ bool FBindableValueBase::ImportTextItem(const TCHAR*& Buffer, int32 PortFlags,
 	bool StructType::ExportTextItem(FString& ValueStr, const StructType& DefaultValue, \
 		UObject* Parent, int32 PortFlags, UObject* ExportRootScope) const \
 	{ \
-		const StructType TypeDefault; \
-		StructType::StaticStruct()->ExportText(ValueStr, this, &TypeDefault, Parent, PortFlags, ExportRootScope, /*bAllowNativeOverride=*/false); \
+		StructType::StaticStruct()->ExportText(ValueStr, this, this, Parent, PortFlags, ExportRootScope, /*bAllowNativeOverride=*/false); \
 		ExportBindingText(ValueStr, Binding.Get(), PortFlags); \
 		return true; \
 	} \
@@ -703,6 +709,7 @@ IMPLEMENT_BINDABLE_TEXT_SERIALIZATION(FBindableByte)
 IMPLEMENT_BINDABLE_TEXT_SERIALIZATION(FBindableName)
 IMPLEMENT_BINDABLE_TEXT_SERIALIZATION(FBindableVector)
 IMPLEMENT_BINDABLE_TEXT_SERIALIZATION(FBindableQuat)
+IMPLEMENT_BINDABLE_TEXT_SERIALIZATION(FBindableTransform)
 IMPLEMENT_BINDABLE_TEXT_SERIALIZATION(FBindableEnum)
 IMPLEMENT_BINDABLE_TEXT_SERIALIZATION(FBindableStruct)
 IMPLEMENT_BINDABLE_TEXT_SERIALIZATION(FBindableObject)
diff --git a/Engine/Plugins/Experimental/UAF/UAF/Source/UAF/Public/BindableValue/UAFBindableTypes.h b/Engine/Plugins/Experimental/UAF/UAF/Source/UAF/Public/BindableValue/UAFBindableTypes.h
--- a/Engine/Plugins/Experimental/UAF/UAF/Source/UAF/Public/BindableValue/UAFBindableTypes.h
+++ b/Engine/Plugins/Experimental/UAF/UAF/Source/UAF/Public/BindableValue/UAFBindableTypes.h
@@ -456,6 +456,41 @@ struct TStructOpsTypeTraits<FBindableQuat> : public TStructOpsTypeTraitsBase2<FB
 	};
 };
 
+//------------------------------------------------------------------------------
+// FBindableTransform
+//------------------------------------------------------------------------------
+
+USTRUCT(BlueprintType, MinimalAPI)
+struct FBindableTransform : public FBindableValueBase
+{
+	GENERATED_BODY()
+	DECLARE_BINDABLE_TEXT_SERIALIZATION(FBindableTransform)
+
+	FBindableTransform() = default;
+	FBindableTransform(const FTransform& InValue) : ConstantValue(InValue) {}
+
+	UAF_API FTransform GetValue(FUAFAssetInstance* Instance) const;
+
+	const FTransform& GetConstantValue() const { return ConstantValue; }
+	void              SetConstantValue(const FTransform& InValue) { ConstantValue = InValue; }
+
+private:
+	UPROPERTY(EditAnywhere, Category = "Value")
+	FTransform ConstantValue = FTransform::Identity;
+};
+
+template<>
+struct TStructOpsTypeTraits<FBindableTransform> : public TStructOpsTypeTraitsBase2<FBindableTransform>
+{
+	enum
+	{
+		WithSerializer = true,
+		WithCopy = true,
+		WithExportTextItem = true,
+		WithImportTextItem = true,
+	};
+};
+
 //------------------------------------------------------------------------------
 // FBindableEnum
 //------------------------------------------------------------------------------
diff --git a/Engine/Plugins/Experimental/UAF/UAF/Source/UAFEditor/Private/AnimNextEditorModule.cpp b/Engine/Plugins/Experimental/UAF/UAF/Source/UAFEditor/Private/AnimNextEditorModule.cpp
--- a/Engine/Plugins/Experimental/UAF/UAF/Source/UAFEditor/Private/AnimNextEditorModule.cpp
+++ b/Engine/Plugins/Experimental/UAF/UAF/Source/UAFEditor/Private/AnimNextEditorModule.cpp
@@ -198,10 +198,11 @@ void FAnimNextEditorModule::StartupModule()
 		FOnGetDetailCustomizationInstance::CreateLambda([] { return MakeShared<FVariableCustomization>(); }));
 
 	static const FName BindableTypes[] = {
-		"BindableBool",   "BindableFloat",  "BindableDouble",
-		"BindableInt32",  "BindableInt64",  "BindableByte",
-		"BindableName",   "BindableVector", "BindableQuat",
-		"BindableEnum",   "BindableStruct", "BindableObject"
+		"BindableBool",      "BindableFloat",  "BindableDouble",
+		"BindableInt32",     "BindableInt64",  "BindableByte",
+		"BindableName",      "BindableVector", "BindableQuat",
+		"BindableTransform", "BindableEnum",   "BindableStruct",
+		"BindableObject"
 	};
 	for (FName TypeName : BindableTypes)
 	{
@@ -616,10 +617,11 @@ void FAnimNextEditorModule::ShutdownModule()
 		PropertyModule.UnregisterCustomClassLayout("UAFRigVMAssetEditorData");
 
 		static const FName BindableTypes[] = {
-			"BindableBool",   "BindableFloat",  "BindableDouble",
-			"BindableInt32",  "BindableInt64",  "BindableByte",
-			"BindableName",   "BindableVector", "BindableQuat",
-			"BindableEnum",   "BindableStruct", "BindableObject"
+			"BindableBool",      "BindableFloat",  "BindableDouble",
+			"BindableInt32",     "BindableInt64",  "BindableByte",
+			"BindableName",      "BindableVector", "BindableQuat",
+			"BindableTransform", "BindableEnum",   "BindableStruct",
+			"BindableObject"
 		};
 		for (FName TypeName : BindableTypes)
 		{
diff --git a/Engine/Plugins/Experimental/UAF/UAF/Source/UAFEditor/Private/BindableValue/UAFBindableValueCustomization.cpp b/Engine/Plugins/Experimental/UAF/UAF/Source/UAFEditor/Private/BindableValue/UAFBindableValueCustomization.cpp
--- a/Engine/Plugins/Experimental/UAF/UAF/Source/UAFEditor/Private/BindableValue/UAFBindableValueCustomization.cpp
+++ b/Engine/Plugins/Experimental/UAF/UAF/Source/UAFEditor/Private/BindableValue/UAFBindableValueCustomization.cpp
@@ -84,6 +84,10 @@ FAnimNextParamType GetBindableParamType(const UScriptStruct* BindableStruct, TSh
 	{
 		return FAnimNextParamType(EValueType::Struct, EContainerType::None, TBaseStructure<FQuat>::Get());
 	}
+	if (BindableStruct == FBindableTransform::StaticStruct())
+	{
+		return FAnimNextParamType(EValueType::Struct, EContainerType::None, TBaseStructure<FTransform>::Get());
+	}
 
 	if (BindableStruct == FBindableEnum::StaticStruct())
 	{
@@ -489,9 +493,31 @@ void FBindableValueCustomization::CustomizeHeader(
 					Handle->GetValue(Value);
 					return Value;
 				})
-				.OnEnumSelectionChanged_Lambda([Handle = ConstantValueHandle](int32 NewValue, ESelectInfo::Type)
+				.OnEnumSelectionChanged_Lambda([Handle = PropertyHandle, StructType = BindableStruct](int32 NewValue, ESelectInfo::Type)
 				{
-					Handle->SetValue(NewValue);
+					// Use copy+export+SetValueFromFormattedString on the parent handle so
+					// the value change and pin default text update happen atomically.
+					// Direct SetValue on the child handle gets overwritten by the workspace
+					// re-importing the old pin default text during the change notification.
+					FString TextValue;
+					Handle->EnumerateConstRawData(
+						[&TextValue, NewValue, StructType](const void* RawData, const int32, const int32) -> bool
+						{
+							const int32 StructSize = StructType->GetStructureSize();
+
+							TArray<uint8> CopyBuf;
+							CopyBuf.SetNumUninitialized(StructSize);
+							StructType->InitializeStruct(CopyBuf.GetData());
+							StructType->CopyScriptStruct(CopyBuf.GetData(), RawData);
+							static_cast<FBindableEnum*>(static_cast<void*>(CopyBuf.GetData()))->SetConstantValue(NewValue);
+
+							StructType->ExportText(TextValue, CopyBuf.GetData(), CopyBuf.GetData(), nullptr, 0, nullptr, /*bAllowNativeOverride=*/true);
+
+							StructType->DestroyStruct(CopyBuf.GetData());
+							return false;
+						});
+
+					Handle->SetValueFromFormattedString(TextValue);
 				});
 		}
 		else if (bIsBindableStruct && ResolvedStructClass)
diff --git a/Engine/Plugins/Experimental/UAF/UAF/Source/UAFTestData/Public/UAFTestVars.h b/Engine/Plugins/Experimental/UAF/UAF/Source/UAFTestData/Public/UAFTestVars.h
--- a/Engine/Plugins/Experimental/UAF/UAF/Source/UAFTestData/Public/UAFTestVars.h
+++ b/Engine/Plugins/Experimental/UAF/UAF/Source/UAFTestData/Public/UAFTestVars.h
@@ -36,6 +36,9 @@ struct FUAFNestedTestStruct
 
 	UPROPERTY()
 	FQuat Quat = FQuat::Identity;
+
+	UPROPERTY()
+	FTransform Transform = FTransform::Identity;
 };
 
 /**
@@ -67,6 +70,9 @@ struct FUAFTestVars
 	UPROPERTY()
 	FQuat QuatVar = FQuat::Identity;
 
+	UPROPERTY()
+	FTransform TransformVar = FTransform::Identity;
+
 	UPROPERTY()
 	FUAFNestedTestStruct NestedVar;
 
diff --git a/Engine/Plugins/Experimental/UAF/UAF/Tests/UAF/Private/BindableValueSerializationTests.cpp b/Engine/Plugins/Experimental/UAF/UAF/Tests/UAF/Private/BindableValueSerializationTests.cpp
--- a/Engine/Plugins/Experimental/UAF/UAF/Tests/UAF/Private/BindableValueSerializationTests.cpp
+++ b/Engine/Plugins/Experimental/UAF/UAF/Tests/UAF/Private/BindableValueSerializationTests.cpp
@@ -249,6 +249,15 @@ TEST_CASE("FBindableQuat.BinarySerialize.WithBinding", "[UAF][unit][MustPass]")
 	CHECK(Dst.HasBinding());
 }
 
+TEST_CASE("FBindableTransform.BinarySerialize.WithBinding", "[UAF][unit][MustPass]")
+{
+	FBindableTransform Src(FTransform(FQuat(1.0, 2.0, 3.0, 4.0), FVector(5.0, 6.0, 7.0)));
+	Src.SetBinding(MakeTestBinding());
+	FBindableTransform Dst = BinaryRoundTrip(Src);
+	CHECK(Dst.GetConstantValue().GetTranslation().X == 5.0);
+	CHECK(Dst.HasBinding());
+}
+
 TEST_CASE("FBindableEnum.BinarySerialize.WithBinding", "[UAF][unit][MustPass]")
 {
 	FBindableEnum Src(EUAFTestEnum::ValueB);
@@ -377,6 +386,15 @@ TEST_CASE("FBindableQuat.TextSerialize.WithBinding", "[UAF][unit][MustPass]")
 	CHECK(Dst.HasBinding());
 }
 
+TEST_CASE("FBindableTransform.TextSerialize.WithBinding", "[UAF][unit][MustPass]")
+{
+	FBindableTransform Src(FTransform(FQuat(1.0, 2.0, 3.0, 4.0), FVector(5.0, 6.0, 7.0)));
+	Src.SetBinding(MakeTestBinding());
+	FBindableTransform Dst = TextRoundTrip(Src);
+	CHECK(Dst.GetConstantValue().GetTranslation().X == 5.0);
+	CHECK(Dst.HasBinding());
+}
+
 TEST_CASE("FBindableEnum.TextSerialize.WithBinding", "[UAF][unit][MustPass]")
 {
 	FBindableEnum Src(EUAFTestEnum::ValueB);
@@ -482,6 +500,14 @@ TEST_CASE("FBindableQuat.EditorSetBinding", "[UAF][unit][MustPass]")
 	CHECK(Dst.HasBinding());
 }
 
+TEST_CASE("FBindableTransform.EditorSetBinding", "[UAF][unit][MustPass]")
+{
+	FBindableTransform Src(FTransform(FQuat(1.0, 2.0, 3.0, 4.0), FVector(5.0, 6.0, 7.0)));
+	FBindableTransform Dst = EditorStyleSetBinding(Src, MakeTestBinding());
+	CHECK(Dst.GetConstantValue().GetTranslation().X == 5.0);
+	CHECK(Dst.HasBinding());
+}
+
 TEST_CASE("FBindableEnum.EditorSetBinding", "[UAF][unit][MustPass]")
 {
 	FBindableEnum Src(EUAFTestEnum::ValueB);
@@ -557,6 +583,15 @@ TEST_CASE("FBindableVector.SelfComparisonExport.WithBinding", "[UAF][unit][MustP
 	CHECK(Dst.HasBinding());
 }
 
+TEST_CASE("FBindableTransform.SelfComparisonExport.WithBinding", "[UAF][unit][MustPass]")
+{
+	FBindableTransform Src(FTransform(FQuat::Identity, FVector(1.0, 2.0, 3.0)));
+	Src.SetBinding(MakeTestBinding());
+	FBindableTransform Dst = SelfComparisonRoundTrip(Src);
+	CHECK(Dst.GetConstantValue().GetTranslation() == FVector(1.0, 2.0, 3.0));
+	CHECK(Dst.HasBinding());
+}
+
 TEST_CASE("FBindableEnum.SelfComparisonExport.WithBinding", "[UAF][unit][MustPass]")
 {
 	FBindableEnum Src(EUAFTestEnum::ValueB);
@@ -588,6 +623,81 @@ TEST_CASE("FBindableVector.EditorClearBinding", "[UAF][unit][MustPass]")
 	CHECK_FALSE(Dst.HasBinding());
 }
 
+TEST_CASE("FBindableTransform.EditorClearBinding", "[UAF][unit][MustPass]")
+{
+	FBindableTransform Src(FTransform(FQuat::Identity, FVector(1.0, 2.0, 3.0)));
+	Src.SetBinding(MakeTestBinding());
+	REQUIRE(Src.HasBinding());
+	FBindableTransform Dst = EditorStyleClearBinding(Src);
+	CHECK(Dst.GetConstantValue().GetTranslation() == FVector(1.0, 2.0, 3.0));
+	CHECK_FALSE(Dst.HasBinding());
+}
+
+// ===========================================================================
+// FBindableEnum zero-value tests — validates that enum values whose underlying
+// int is 0 survive text serialization. Previously ConstantValue=0 was omitted
+// because it matched the default int32 value; the round-trip result would still
+// be 0 "by accident" (it is the default), so we directly assert that the
+// exported text contains the ConstantValue field.
+// ===========================================================================
+
+TEST_CASE("FBindableEnum.TextSerialize.ZeroValueIsExported", "[UAF][unit][MustPass]")
+{
+	FBindableEnum Src(EUAFTestEnum::ValueA); // ValueA == 0
+
+	FString Text;
+	FBindableEnum Default;
+	Src.ExportTextItem(Text, Default, nullptr, 0, nullptr);
+
+	// The exported text must explicitly contain ConstantValue=0, not omit it
+	CHECK(Text.Contains(TEXT("ConstantValue=0")));
+	CHECK(Text.Contains(TEXT("EnumClass=")));
+
+	// Verify round-trip import also works
+	FBindableEnum Dst;
+	const TCHAR* Buffer = *Text;
+	REQUIRE(Dst.ImportTextItem(Buffer, 0, nullptr, nullptr));
+	CHECK(Dst.GetConstantValue() == static_cast<int32>(EUAFTestEnum::ValueA));
+	CHECK(Dst.GetEnumClass() == StaticEnum<EUAFTestEnum>());
+}
+
+TEST_CASE("FBindableEnum.TextSerialize.ZeroValueOverwritesNonZero", "[UAF][unit][MustPass]")
+{
+	// Export ValueA (== 0), then import into a struct that already has ValueB (== 1).
+	// If ConstantValue=0 was omitted from the export, the import would leave
+	// the destination at ValueB — proving the bug.
+	FBindableEnum Src(EUAFTestEnum::ValueA);
+
+	FString Text;
+	FBindableEnum Default;
+	Src.ExportTextItem(Text, Default, nullptr, 0, nullptr);
+
+	FBindableEnum Dst(EUAFTestEnum::ValueB);
+	const TCHAR* Buffer = *Text;
+	REQUIRE(Dst.ImportTextItem(Buffer, 0, nullptr, nullptr));
+	CHECK(Dst.GetConstantValue() == static_cast<int32>(EUAFTestEnum::ValueA));
+	CHECK(Dst.GetEnumClass() == StaticEnum<EUAFTestEnum>());
+}
+
+TEST_CASE("FBindableEnum.SelfComparisonExport.ZeroValueIsExported", "[UAF][unit][MustPass]")
+{
+	// The RigVM path passes Data as both value and default
+	FBindableEnum Src(EUAFTestEnum::ValueA); // ValueA == 0
+
+	FString Text;
+	Src.ExportTextItem(Text, Src, nullptr, 0, nullptr);
+
+	CHECK_FALSE(Text.StartsWith(TEXT("(,")));
+	CHECK(Text.Contains(TEXT("ConstantValue=0")));
+	CHECK(Text.Contains(TEXT("EnumClass=")));
+
+	FBindableEnum Dst;
+	const TCHAR* Buffer = *Text;
+	REQUIRE(Dst.ImportTextItem(Buffer, 0, nullptr, nullptr));
+	CHECK(Dst.GetConstantValue() == static_cast<int32>(EUAFTestEnum::ValueA));
+	CHECK(Dst.GetEnumClass() == StaticEnum<EUAFTestEnum>());
+}
+
 // ===========================================================================
 // Sibling binding isolation tests — when multiple FBindableXxx members are
 // inside the same parent struct, ImportTextItem receives the full remaining
@@ -647,3 +757,16 @@ TEST_CASE("FBindableVector.ImportTextItem.DoesNotPickUpSiblingBinding", "[UAF][u
 	CHECK(Dst.GetConstantValue() == FVector(1.0, 2.0, 3.0));
 	CHECK_FALSE(Dst.HasBinding());
 }
+
+TEST_CASE("FBindableTransform.ImportTextItem.DoesNotPickUpSiblingBinding", "[UAF][unit][MustPass]")
+{
+	// FTransform has deeply nested parens — verify struct boundary detection handles it
+	const FString Buffer = TEXT("(ConstantValue=(Rotation=(X=0.000000,Y=0.000000,Z=0.000000,W=1.000000),Translation=(X=1.000000,Y=2.000000,Z=3.000000),Scale3D=(X=1.000000,Y=1.000000,Z=1.000000))),Sibling=(ConstantValue=0,__Binding=(SourceType=Variable,SourceVariable=(Name=\"FloatVal\")))");
+	const TCHAR* Ptr = *Buffer;
+
+	FBindableTransform Dst;
+	const bool bImported = Dst.ImportTextItem(Ptr, 0, nullptr, nullptr);
+	REQUIRE(bImported);
+	CHECK(Dst.GetConstantValue().GetTranslation() == FVector(1.0, 2.0, 3.0));
+	CHECK_FALSE(Dst.HasBinding());
+}
diff --git a/Engine/Plugins/Experimental/UAF/UAF/Tests/UAF/Private/BindableValueTests.cpp b/Engine/Plugins/Experimental/UAF/UAF/Tests/UAF/Private/BindableValueTests.cpp
--- a/Engine/Plugins/Experimental/UAF/UAF/Tests/UAF/Private/BindableValueTests.cpp
+++ b/Engine/Plugins/Experimental/UAF/UAF/Tests/UAF/Private/BindableValueTests.cpp
@@ -491,6 +491,44 @@ TEST_CASE("FBindableQuat.SubPropertyNestedQuatResolvesToQuat", "[UAF][unit][Must
 	REQUIRE(Field.GetValue(&Owner) == NestedData.Quat);
 }
 
+// ---------------------------------------------------------------------------
+// FBindableTransform
+// ---------------------------------------------------------------------------
+
+TEST_CASE("FBindableTransform.NoBindingReturnsConstantValue", "[UAF][unit][MustPass]")
+{
+	const FTransform Expected(FQuat(FVector::UpVector, UE_HALF_PI), FVector(1.0, 2.0, 3.0), FVector(1.0, 1.0, 1.0));
+	FBindableTransform Field;
+	Field.SetConstantValue(Expected);
+
+	REQUIRE(Field.GetValue(nullptr).Equals(Expected));
+}
+
+TEST_CASE("FBindableTransform.VariableBindingResolvesTransform", "[UAF][unit][MustPass]")
+{
+	auto [Owner, TransformRef] = MakeOwnerAndRef("TransformVar");
+	const FTransform Source(FQuat(FVector::UpVector, UE_HALF_PI), FVector(10.0, 20.0, 30.0), FVector(1.0, 1.0, 1.0));
+	REQUIRE(Owner.SetVariable(TransformRef, Source) == EPropertyBagResult::Success);
+
+	FBindableTransform Field;
+	Field.SetBinding(MakeVariablePropertyBinding(TransformRef));
+
+	REQUIRE(Field.GetValue(&Owner).Equals(Source));
+}
+
+TEST_CASE("FBindableTransform.SubPropertyNestedTransformResolvesToTransform", "[UAF][unit][MustPass]")
+{
+	auto [Owner, NestedRef] = MakeOwnerAndRef("NestedVar");
+	FUAFNestedTestStruct NestedData;
+	NestedData.Transform = FTransform(FQuat(FVector::ForwardVector, UE_HALF_PI), FVector(7.0, 8.0, 9.0));
+	REQUIRE(Owner.SetVariable(NestedRef, NestedData) == EPropertyBagResult::Success);
+
+	FBindableTransform Field;
+	Field.SetBinding(MakeSubPropertyBinding(NestedRef, TEXT("Transform")));
+
+	REQUIRE(Field.GetValue(&Owner).Equals(NestedData.Transform));
+}
+
 // ---------------------------------------------------------------------------
 // FBindableStruct
 // ---------------------------------------------------------------------------
diff --git a/Engine/Plugins/Experimental/UAF/UAFTestSuites/Source/UAFAnimNodeTestData/Public/UAFAnimNodeTestVars.h b/Engine/Plugins/Experimental/UAF/UAFTestSuites/Source/UAFAnimNodeTestData/Public/UAFAnimNodeTestVars.h
--- a/Engine/Plugins/Experimental/UAF/UAFTestSuites/Source/UAFAnimNodeTestData/Public/UAFAnimNodeTestVars.h
+++ b/Engine/Plugins/Experimental/UAF/UAFTestSuites/Source/UAFAnimNodeTestData/Public/UAFAnimNodeTestVars.h
@@ -19,6 +19,9 @@ struct FUAFAnimNodeNestedTestStruct
 
 	UPROPERTY()
 	FQuat Quat = FQuat::Identity;
+
+	UPROPERTY()
+	FTransform Transform = FTransform::Identity;
 };
 
 /**
@@ -50,6 +53,9 @@ struct FUAFAnimNodeTestVars
 	UPROPERTY(EditAnywhere, Category = "Variables")
 	FQuat QuatVar = FQuat::Identity;
 
+	UPROPERTY(EditAnywhere, Category = "Variables")
+	FTransform TransformVar = FTransform::Identity;
+
 	UPROPERTY(EditAnywhere, Category = "Variables")
 	FUAFAnimNodeNestedTestStruct NestedVar;
 };
@@ -132,6 +138,9 @@ struct FUAFTestAnimNodeData : public FUAFAnimNodeData
 	UPROPERTY(EditAnywhere, Category = "NodeData")
 	FBindableQuat QuatVal = FQuat::Identity;
 
+	UPROPERTY(EditAnywhere, Category = "NodeData")
+	FBindableTransform TransformBindableVal = FTransform::Identity;
+
 	UPROPERTY(EditAnywhere, Category = "NodeData")
 	FBindableObject ObjectVal = FBindableObject();
 
