Chaos API: Updating shape instance to handle materials.

#Synced-CL 52508799

diff --git a/Engine/Plugins/Experimental/ChaosRigidPhysicsAsync/Source/ChaosRigidPhysicsAsync/Private/ChaosRigidPhysicsAsync/RigidShapeInstanceAsync.cpp b/Engine/Plugins/Experimental/ChaosRigidPhysicsAsync/Source/ChaosRigidPhysicsAsync/Private/ChaosRigidPhysicsAsync/RigidShapeInstanceAsync.cpp
--- a/Engine/Plugins/Experimental/ChaosRigidPhysicsAsync/Source/ChaosRigidPhysicsAsync/Private/ChaosRigidPhysicsAsync/RigidShapeInstanceAsync.cpp
+++ b/Engine/Plugins/Experimental/ChaosRigidPhysicsAsync/Source/ChaosRigidPhysicsAsync/Private/ChaosRigidPhysicsAsync/RigidShapeInstanceAsync.cpp
@@ -233,6 +233,25 @@ namespace Chaos::Rigids::Async
 			return nullptr;
 		}
 
+		TArray<Chaos::FMaterialHandle> MaterialHandles;
+		MaterialHandles.Reserve(InSetup.Materials.Num());
+		for (const UE::Physics::FMaterialHandle& Handle : InSetup.Materials)
+		{
+			MaterialHandles.Add(Handle.GetMaterialHandle());
+		}
+		TArray<Chaos::FMaterialMaskHandle> MaskHandles;
+		MaskHandles.Reserve(InSetup.MaterialMaskHandles.Num());
+		for (const UE::Physics::FMaterialMaskHandle& Handle : InSetup.MaterialMaskHandles)
+		{
+			MaskHandles.Add(Handle.GetMaterialMaskHandle());
+		}
+		TArray<Chaos::FMaterialHandle> MaskMapMaterialHandles;
+		MaskMapMaterialHandles.Reserve(InSetup.MaterialMaskMapMaterials.Num());
+		for (const UE::Physics::FMaterialHandle& Handle : InSetup.MaterialMaskMapMaterials)
+		{
+			MaskMapMaterialHandles.Add(Handle.GetMaterialHandle());
+		}
+
 		TUniquePtr<Chaos::FPerShapeData> NewShape = Chaos::FShapeInstanceProxy::Make(0, ImplicitPtr);
 		Chaos::FPerShapeData* PerShapeData = NewShape.Release();
 
@@ -243,8 +262,11 @@ namespace Chaos::Rigids::Async
 		PerShapeData->SetFilterInstanceData(InSetup.FilterInstanceData);
 		PerShapeData->SetCollisionTraceType(ConvertCollisionTraceFlag(InSetup.CollisionTraceType));
 		// TODO: UpdateShapeBounds?
+		PerShapeData->SetMaterials(MoveTemp(MaterialHandles));
+		PerShapeData->SetMaterialMasks(MoveTemp(MaskHandles));
+		PerShapeData->SetMaterialMaskMaps(InSetup.MaterialMaskMaps);
+		PerShapeData->SetMaterialMaskMapMaterials(MoveTemp(MaskMapMaterialHandles));
 		PerShapeData->SetUserData(InSetup.UserData);
-		// TODO: Materials
 
 		return PerShapeData;
 	}
@@ -427,6 +449,84 @@ namespace Chaos::Rigids::Async
 	{
 		PerShapeData->SetFilterInstanceData(InInstanceData);
 	}
+	
+	int32 FRigidShapeInstanceAsync::GetNumMaterials() const
+	{
+		return PerShapeData->NumMaterials();
+	}
+
+	UE::Physics::FMaterialHandle FRigidShapeInstanceAsync::GetMaterial(const int32 InIndex) const
+	{
+		return UE::Physics::FMaterialHandle(PerShapeData->GetMaterial(InIndex));
+	}
+
+	void FRigidShapeInstanceAsync::SetMaterials(TArray<UE::Physics::FMaterialHandle>&& InMaterials)
+	{
+		TArray<Chaos::FMaterialHandle> MaterialHandles;
+		MaterialHandles.Reserve(InMaterials.Num());
+		for (const UE::Physics::FMaterialHandle& MaterialHandle : InMaterials)
+		{
+			MaterialHandles.Add(MaterialHandle.GetMaterialHandle());
+		}
+		PerShapeData->SetMaterials(MoveTemp(MaterialHandles));
+	}
+
+	int32 FRigidShapeInstanceAsync::GetNumMaterialMasks() const
+	{
+		return PerShapeData->GetMaterialMasks().Num();
+	}
+
+	UE::Physics::FMaterialMaskHandle FRigidShapeInstanceAsync::GetMaterialMask(const int32 InIndex) const
+	{
+		return UE::Physics::FMaterialMaskHandle(PerShapeData->GetMaterialMasks()[InIndex]);
+	}
+
+	void FRigidShapeInstanceAsync::SetMaterialMasks(TArray<UE::Physics::FMaterialMaskHandle>&& InMaterialMasks)
+	{
+		TArray<Chaos::FMaterialMaskHandle> MaskHandles;
+		MaskHandles.Reserve(InMaterialMasks.Num());
+		for (const UE::Physics::FMaterialMaskHandle& MaterialMaskHandle : InMaterialMasks)
+		{
+			MaskHandles.Add(MaterialMaskHandle.GetMaterialMaskHandle());
+		}
+		PerShapeData->SetMaterialMasks(MoveTemp(MaskHandles));
+	}
+
+	int32 FRigidShapeInstanceAsync::GetNumMaterialMaskMaps() const
+	{
+		return PerShapeData->GetMaterialMaskMaps().Num();
+	}
+
+	uint32 FRigidShapeInstanceAsync::GetMaterialMaskMap(const int32 InIndex) const
+	{
+		return PerShapeData->GetMaterialMaskMaps()[InIndex];
+	}
+
+	void FRigidShapeInstanceAsync::SetMaterialMaskMaps(TArray<uint32>&& InMaterialMaskMaps)
+	{
+		PerShapeData->SetMaterialMaskMaps(MoveTemp(InMaterialMaskMaps));
+	}
+
+	int32 FRigidShapeInstanceAsync::GetNumMaterialMaskMapMaterials() const
+	{
+		return PerShapeData->GetMaterialMaskMapMaterials().Num();
+	}
+
+	UE::Physics::FMaterialHandle FRigidShapeInstanceAsync::GetMaterialMaskMapMaterial(const int32 InIndex) const
+	{
+		return UE::Physics::FMaterialHandle(PerShapeData->GetMaterialMaskMapMaterials()[InIndex]);
+	}
+
+	void FRigidShapeInstanceAsync::SetMaterialMaskMapMaterials(TArray<UE::Physics::FMaterialHandle>&& InMaterialMaskMapMaterials)
+	{
+		TArray<Chaos::FMaterialHandle> MaskMapMaterialHandles;
+		MaskMapMaterialHandles.Reserve(InMaterialMaskMapMaterials.Num());
+		for (const UE::Physics::FMaterialHandle& MaterialHandle : InMaterialMaskMapMaterials)
+		{
+			MaskMapMaterialHandles.Add(MaterialHandle.GetMaterialHandle());
+		}
+		PerShapeData->SetMaterialMaskMapMaterials(MoveTemp(MaskMapMaterialHandles));
+	}
 
 	void* FRigidShapeInstanceAsync::GetUserData() const
 	{
diff --git a/Engine/Plugins/Experimental/ChaosRigidPhysicsAsync/Source/ChaosRigidPhysicsAsync/Public/ChaosRigidPhysicsAsync/RigidShapeInstanceAsync.h b/Engine/Plugins/Experimental/ChaosRigidPhysicsAsync/Source/ChaosRigidPhysicsAsync/Public/ChaosRigidPhysicsAsync/RigidShapeInstanceAsync.h
--- a/Engine/Plugins/Experimental/ChaosRigidPhysicsAsync/Source/ChaosRigidPhysicsAsync/Public/ChaosRigidPhysicsAsync/RigidShapeInstanceAsync.h
+++ b/Engine/Plugins/Experimental/ChaosRigidPhysicsAsync/Source/ChaosRigidPhysicsAsync/Public/ChaosRigidPhysicsAsync/RigidShapeInstanceAsync.h
@@ -44,6 +44,22 @@ namespace Chaos::Rigids::Async
 
 		UE_INTERNAL virtual Chaos::Filter::FInstanceData GetFilterInstanceData() const override final;
 		UE_INTERNAL virtual void SetFilterInstanceData(const Chaos::Filter::FInstanceData& InInstanceData) override final;
+		
+		UE_INTERNAL virtual int32 GetNumMaterials() const override final;
+		UE_INTERNAL virtual UE::Physics::FMaterialHandle GetMaterial(const int32 InIndex) const override final;
+		UE_INTERNAL virtual void SetMaterials(TArray<UE::Physics::FMaterialHandle>&& InMaterials) override final;
+
+		UE_INTERNAL virtual int32 GetNumMaterialMasks() const override final;
+		UE_INTERNAL virtual UE::Physics::FMaterialMaskHandle GetMaterialMask(const int32 InIndex) const override final;
+		UE_INTERNAL virtual void SetMaterialMasks(TArray<UE::Physics::FMaterialMaskHandle>&& InMaterialMasks) override final;
+
+		UE_INTERNAL virtual int32 GetNumMaterialMaskMaps() const override final;
+		UE_INTERNAL virtual uint32 GetMaterialMaskMap(const int32 InIndex) const override final;
+		UE_INTERNAL virtual void SetMaterialMaskMaps(TArray<uint32>&& InMaterialMaskMaps) override final;
+
+		UE_INTERNAL virtual int32 GetNumMaterialMaskMapMaterials() const override final;
+		UE_INTERNAL virtual UE::Physics::FMaterialHandle GetMaterialMaskMapMaterial(const int32 InIndex) const override final;
+		UE_INTERNAL virtual void SetMaterialMaskMapMaterials(TArray<UE::Physics::FMaterialHandle>&& InMaterialMaskMapMaterials) override final;
 
 		UE_INTERNAL virtual void* GetUserData() const override final;
 		UE_INTERNAL virtual void SetUserData(void* InUserData) override final;
diff --git a/Engine/Plugins/Experimental/ChaosRigidPhysicsAsync/Tests/Private/RigidTestFixture.cpp b/Engine/Plugins/Experimental/ChaosRigidPhysicsAsync/Tests/Private/RigidTestFixture.cpp
--- a/Engine/Plugins/Experimental/ChaosRigidPhysicsAsync/Tests/Private/RigidTestFixture.cpp
+++ b/Engine/Plugins/Experimental/ChaosRigidPhysicsAsync/Tests/Private/RigidTestFixture.cpp
@@ -45,6 +45,14 @@ namespace Chaos::LowLevelTest
 				Context->DestroyGeometryCollection(GCHandle.Pin(Context));
 			}
 		}
+		for (FMaterialHandle& MaterialHandle : MaterialHandles)
+		{
+			FPhysicalMaterialManager::Get().Destroy(MaterialHandle.GetMaterialHandle());
+		}
+		for (FMaterialMaskHandle& MaterialMaskHandle : MaterialMaskHandles)
+		{
+			FPhysicalMaterialManager::Get().Destroy(MaterialMaskHandle.GetMaterialMaskHandle());
+		}
 		Factory.DestroyScene(SceneHandle);
 	}
 
@@ -62,6 +70,26 @@ namespace Chaos::LowLevelTest
 		return SceneHandle.LockRW();
 	}
 
+	FMaterialHandle FRigidTestFixture::CreateMaterial(bool bAutoCleanup)
+	{
+		FMaterialHandle MaterialHandle(FPhysicalMaterialManager::Get().Create());
+		if (bAutoCleanup)
+		{
+			MaterialHandles.Add(MaterialHandle);
+		}
+		return MaterialHandle;
+	}
+
+	FMaterialMaskHandle FRigidTestFixture::CreateMaterialMask(bool bAutoCleanup)
+	{
+		FMaterialMaskHandle MaterialMaskHandle(FPhysicalMaterialManager::Get().CreateMask());
+		if (bAutoCleanup)
+		{
+			MaterialMaskHandles.Add(MaterialMaskHandle);
+		}
+		return MaterialMaskHandle;
+	}
+
 	void FRigidTestFixture::AutoCleanup(FJointConstraintHandle Handle)
 	{
 		JointHandles.Add(Handle);
diff --git a/Engine/Plugins/Experimental/ChaosRigidPhysicsAsync/Tests/Private/RigidTestFixture.h b/Engine/Plugins/Experimental/ChaosRigidPhysicsAsync/Tests/Private/RigidTestFixture.h
--- a/Engine/Plugins/Experimental/ChaosRigidPhysicsAsync/Tests/Private/RigidTestFixture.h
+++ b/Engine/Plugins/Experimental/ChaosRigidPhysicsAsync/Tests/Private/RigidTestFixture.h
@@ -5,6 +5,7 @@
 #include "ChaosRigidPhysicsAsync/RigidFactoryAsync.h"
 #include "ChaosRigidPhysicsAsync/RigidSceneSettingsAsync.h"
 #include "RigidPhysics/RigidScene.h"
+#include "RigidPhysics/RigidMaterials.h"
 
 #if UE_RIGIDPHYSICS_API_ENABLED
 
@@ -19,6 +20,8 @@ namespace Chaos::LowLevelTest
 	using FRigidGeometryCollectionHandle = UE::Physics::FRigidGeometryCollectionHandle;
 	using FRigidContextGameRO = UE::Physics::FRigidContextGameRO;
 	using FRigidContextGameRW = UE::Physics::FRigidContextGameRW;
+	using FMaterialHandle = UE::Physics::FMaterialHandle;
+	using FMaterialMaskHandle = UE::Physics::FMaterialMaskHandle;
 
 	struct FRigidTestFixture
 	{
@@ -28,6 +31,9 @@ namespace Chaos::LowLevelTest
 		FRigidContextGameRO LockROChecked() const;
 		FRigidContextGameRW LockRWChecked() const;
 
+		FMaterialHandle CreateMaterial(bool bAutoCleanup = true);
+		FMaterialMaskHandle CreateMaterialMask(bool bAutoCleanup = true);
+
 		void AutoCleanup(FJointConstraintHandle Handle);
 		void AutoCleanup(FRigidBodyHandle Handle);
 		void AutoCleanup(FRigidGeometryCollectionHandle Handle);
@@ -41,6 +47,8 @@ namespace Chaos::LowLevelTest
 		TArray<FRigidBodyHandle> BodyHandles;
 		TArray<FRigidGeometryCollectionHandle> GeometryCollectionHandles;
 		TArray<FJointConstraintHandle> JointHandles;
+		TArray<FMaterialHandle> MaterialHandles;
+		TArray<FMaterialMaskHandle> MaterialMaskHandles;
 	};
 } // namespace Chaos::LowLevelTest
 
diff --git a/Engine/Plugins/Experimental/ChaosRigidPhysicsAsync/Tests/Private/ShapeInstanceTests.cpp b/Engine/Plugins/Experimental/ChaosRigidPhysicsAsync/Tests/Private/ShapeInstanceTests.cpp
--- a/Engine/Plugins/Experimental/ChaosRigidPhysicsAsync/Tests/Private/ShapeInstanceTests.cpp
+++ b/Engine/Plugins/Experimental/ChaosRigidPhysicsAsync/Tests/Private/ShapeInstanceTests.cpp
@@ -607,6 +607,199 @@ namespace Chaos::LowLevelTest
 			});
 	}
 
+	TEST_CASE("ShapeInstance: Materials", "[Chaos][API][ShapeInstance][unit]")
+	{
+		FShapeInstanceTestFixture Fixture;
+
+		TArray<FMaterialHandle> Materials
+		{
+			Fixture.CreateMaterial(),
+		};
+
+		FTriangleMeshGeometry Geometry(UE::Physics::MakeQuadTriangleGeometrySetup());
+		FRigidShapeInstanceSetup Setup(Geometry);
+		Setup.Materials = Materials;
+
+		FRigidShapeInstanceHandle ShapeInstanceHandle = Fixture.CreateBodyAndShapeInstance(Setup);
+		auto CheckMaterialData = [ShapeInstanceHandle, &Materials]<typename ContextType>(const ContextType& Context)
+		{
+			TRigidShapeInstancePtr<ContextType> ShapeInstance = ShapeInstanceHandle.Pin(Context);
+			REQUIRE(ShapeInstance.IsValid());
+			REQUIRE(Materials.Num() == ShapeInstance->GetNumMaterials());
+			for (int32 I = 0; I < Materials.Num(); ++I)
+			{
+				CHECK(Materials[I] == ShapeInstance->GetMaterial(I));
+			}
+		};
+
+		if (FRigidContextGameRO Context = Fixture.LockROChecked())
+		{
+			CheckMaterialData(Context);
+		}
+		Fixture.RunPTCallback(CheckMaterialData);
+
+		// Now check the setter
+		Materials = TArray<FMaterialHandle>
+		{
+			Fixture.CreateMaterial(),
+			Fixture.CreateMaterial(),
+		};
+		if (FRigidContextGameRW Context = Fixture.LockRWChecked())
+		{
+			TRigidShapeInstancePtr<FRigidContextGameRW> ShapeInstance = ShapeInstanceHandle.Pin(Context);
+			ShapeInstance->SetMaterials(TArray(Materials));
+		}
+
+		if (FRigidContextGameRO Context = Fixture.LockROChecked())
+		{
+			CheckMaterialData(Context);
+		}
+		Fixture.RunPTCallback(CheckMaterialData);
+	}
+
+	TEST_CASE("ShapeInstance: MaterialMasks", "[Chaos][API][ShapeInstance][unit]")
+	{
+		FShapeInstanceTestFixture Fixture;
+
+		TArray<FMaterialMaskHandle> MaterialMasks
+		{
+			Fixture.CreateMaterialMask(),
+		};
+
+		FTriangleMeshGeometry Geometry(UE::Physics::MakeQuadTriangleGeometrySetup());
+		FRigidShapeInstanceSetup Setup(Geometry);
+		Setup.MaterialMaskHandles = MaterialMasks;
+
+		FRigidShapeInstanceHandle ShapeInstanceHandle = Fixture.CreateBodyAndShapeInstance(Setup);
+		auto CheckMaterialMaskData = [ShapeInstanceHandle, &MaterialMasks]<typename ContextType>(const ContextType& Context)
+		{
+			TRigidShapeInstancePtr<ContextType> ShapeInstance = ShapeInstanceHandle.Pin(Context);
+			REQUIRE(ShapeInstance.IsValid());
+			REQUIRE(MaterialMasks.Num() == ShapeInstance->GetNumMaterialMasks());
+			for (int32 I = 0; I < MaterialMasks.Num(); ++I)
+			{
+				CHECK(MaterialMasks[I] == ShapeInstance->GetMaterialMask(I));
+			}
+		};
+
+		if (FRigidContextGameRO Context = Fixture.LockROChecked())
+		{
+			CheckMaterialMaskData(Context);
+		}
+		Fixture.RunPTCallback(CheckMaterialMaskData);
+
+		// Now check the setter
+		MaterialMasks = TArray<FMaterialMaskHandle>
+		{
+			Fixture.CreateMaterialMask(),
+			Fixture.CreateMaterialMask(),
+		};
+		if (FRigidContextGameRW Context = Fixture.LockRWChecked())
+		{
+			TRigidShapeInstancePtr<FRigidContextGameRW> ShapeInstance = ShapeInstanceHandle.Pin(Context);
+			ShapeInstance->SetMaterialMasks(TArray(MaterialMasks));
+		}
+
+		if (FRigidContextGameRO Context = Fixture.LockROChecked())
+		{
+			CheckMaterialMaskData(Context);
+		}
+		Fixture.RunPTCallback(CheckMaterialMaskData);
+	}
+
+	TEST_CASE("ShapeInstance: MaterialMaskMaps", "[Chaos][API][ShapeInstance][unit]")
+	{
+		FShapeInstanceTestFixture Fixture;
+
+		TArray<uint32> MaterialMaskMaps{ 1, 2 };
+
+		FTriangleMeshGeometry Geometry(UE::Physics::MakeQuadTriangleGeometrySetup());
+		FRigidShapeInstanceSetup Setup(Geometry);
+		Setup.MaterialMaskMaps = MaterialMaskMaps;
+
+		FRigidShapeInstanceHandle ShapeInstanceHandle = Fixture.CreateBodyAndShapeInstance(Setup);
+		auto CheckMaterialMaskMapData = [ShapeInstanceHandle, &MaterialMaskMaps]<typename ContextType>(const ContextType& Context)
+		{
+			TRigidShapeInstancePtr<ContextType> ShapeInstance = ShapeInstanceHandle.Pin(Context);
+			REQUIRE(ShapeInstance.IsValid());
+			REQUIRE(MaterialMaskMaps.Num() == ShapeInstance->GetNumMaterialMaskMaps());
+			for (int32 I = 0; I < MaterialMaskMaps.Num(); ++I)
+			{
+				CHECK(MaterialMaskMaps[I] == ShapeInstance->GetMaterialMaskMap(I));
+			}
+		};
+
+		if (FRigidContextGameRO Context = Fixture.LockROChecked())
+		{
+			CheckMaterialMaskMapData(Context);
+		}
+		Fixture.RunPTCallback(CheckMaterialMaskMapData);
+
+		// Now check the setter
+		MaterialMaskMaps = TArray<uint32>{ 3, 4, 5 };
+		if (FRigidContextGameRW Context = Fixture.LockRWChecked())
+		{
+			TRigidShapeInstancePtr<FRigidContextGameRW> ShapeInstance = ShapeInstanceHandle.Pin(Context);
+			ShapeInstance->SetMaterialMaskMaps(TArray(MaterialMaskMaps));
+		}
+
+		if (FRigidContextGameRO Context = Fixture.LockROChecked())
+		{
+			CheckMaterialMaskMapData(Context);
+		}
+		Fixture.RunPTCallback(CheckMaterialMaskMapData);
+	}
+
+	TEST_CASE("ShapeInstance: MaterialMaskMapMaterials", "[Chaos][API][ShapeInstance][unit]")
+	{
+		FShapeInstanceTestFixture Fixture;
+
+		TArray<FMaterialHandle> MaterialMaskMapMaterials
+		{
+			Fixture.CreateMaterial(),
+		};
+
+		FTriangleMeshGeometry Geometry(UE::Physics::MakeQuadTriangleGeometrySetup());
+		FRigidShapeInstanceSetup Setup(Geometry);
+		Setup.MaterialMaskMapMaterials = MaterialMaskMapMaterials;
+
+		FRigidShapeInstanceHandle ShapeInstanceHandle = Fixture.CreateBodyAndShapeInstance(Setup);
+		auto CheckMaterialMaskMapMaterialData = [ShapeInstanceHandle, &MaterialMaskMapMaterials]<typename ContextType>(const ContextType& Context)
+		{
+			TRigidShapeInstancePtr<ContextType> ShapeInstance = ShapeInstanceHandle.Pin(Context);
+			REQUIRE(ShapeInstance.IsValid());
+			REQUIRE(MaterialMaskMapMaterials.Num() == ShapeInstance->GetNumMaterialMaskMapMaterials());
+			for (int32 I = 0; I < MaterialMaskMapMaterials.Num(); ++I)
+			{
+				CHECK(MaterialMaskMapMaterials[I] == ShapeInstance->GetMaterialMaskMapMaterial(I));
+			}
+		};
+
+		if (FRigidContextGameRO Context = Fixture.LockROChecked())
+		{
+			CheckMaterialMaskMapMaterialData(Context);
+		}
+		Fixture.RunPTCallback(CheckMaterialMaskMapMaterialData);
+
+		// Now check the setter
+		MaterialMaskMapMaterials = TArray<FMaterialHandle>
+		{
+			Fixture.CreateMaterial(),
+			Fixture.CreateMaterial(),
+		};
+		if (FRigidContextGameRW Context = Fixture.LockRWChecked())
+		{
+			TRigidShapeInstancePtr<FRigidContextGameRW> ShapeInstance = ShapeInstanceHandle.Pin(Context);
+			ShapeInstance->SetMaterialMaskMapMaterials(TArray(MaterialMaskMapMaterials));
+		}
+
+		if (FRigidContextGameRO Context = Fixture.LockROChecked())
+		{
+			CheckMaterialMaskMapMaterialData(Context);
+		}
+		Fixture.RunPTCallback(CheckMaterialMaskMapMaterialData);
+	}
+
 	TEST_CASE("ShapeInstance: UserData", "[Chaos][API][ShapeInstance][unit]")
 	{
 		FShapeInstanceTestFixture Fixture;
diff --git a/Engine/Source/Runtime/Experimental/RigidPhysics/Private/RigidPhysics/RigidMaterials.cpp b/Engine/Source/Runtime/Experimental/RigidPhysics/Private/RigidPhysics/RigidMaterials.cpp
new file mode 100644
--- /dev/null
+++ b/Engine/Source/Runtime/Experimental/RigidPhysics/Private/RigidPhysics/RigidMaterials.cpp
@@ -0,0 +1,45 @@
+// Copyright Epic Games, Inc. All Rights Reserved.
+#include "RigidPhysics/RigidMaterials.h"
+
+#include "Chaos/PhysicalMaterials.h"
+
+#if UE_RIGIDPHYSICS_API_ENABLED
+
+namespace UE::Physics
+{
+	FMaterialHandle::FMaterialHandle(const Chaos::FMaterialHandle& InMaterialHandle)
+		: HandleData(InMaterialHandle.InnerHandle.AsUint())
+	{
+	}
+
+	Chaos::FMaterialHandle FMaterialHandle::GetMaterialHandle() const
+	{
+		Chaos::FMaterialHandle Result;
+		Result.InnerHandle.FromUint(HandleData);
+		return Result;
+	}
+
+	Chaos::FChaosPhysicsMaterial* FMaterialHandle::GetMaterial() const
+	{
+		return GetMaterialHandle().Get();
+	}
+
+	FMaterialMaskHandle::FMaterialMaskHandle(const Chaos::FMaterialMaskHandle& InMaterialMaskHandle)
+		: HandleData(InMaterialMaskHandle.InnerHandle.AsUint())
+	{
+	}
+
+	Chaos::FMaterialMaskHandle FMaterialMaskHandle::GetMaterialMaskHandle() const
+	{
+		Chaos::FMaterialMaskHandle Result;
+		Result.InnerHandle.FromUint(HandleData);
+		return Result;
+	}
+
+	Chaos::FChaosPhysicsMaterialMask* FMaterialMaskHandle::GetMaterialMask() const
+	{
+		return GetMaterialMaskHandle().Get();
+	}
+} // namespace UE::Physics
+
+#endif // UE_RIGIDPHYSICS_API_ENABLED
diff --git a/Engine/Source/Runtime/Experimental/RigidPhysics/Public/RigidPhysics/Internal/IRigidShapeInstance.h b/Engine/Source/Runtime/Experimental/RigidPhysics/Public/RigidPhysics/Internal/IRigidShapeInstance.h
--- a/Engine/Source/Runtime/Experimental/RigidPhysics/Public/RigidPhysics/Internal/IRigidShapeInstance.h
+++ b/Engine/Source/Runtime/Experimental/RigidPhysics/Public/RigidPhysics/Internal/IRigidShapeInstance.h
@@ -9,8 +9,9 @@
 #include "Chaos/CollisionFilterData.h"
 #include "Math/Transform.h"
 #include "RigidPhysics/Geometry/AnyGeometry.h"
-#include "RigidPhysics/RigidShapeInstanceHandle.h"
 #include "RigidPhysics/RigidLog.h"
+#include "RigidPhysics/RigidMaterials.h"
+#include "RigidPhysics/RigidShapeInstanceHandle.h"
 #include "RigidPhysics/RigidTyped.h"
 
 UE_EXPERIMENTAL(5.8, "The new Chaos API is experimental")
@@ -49,6 +50,22 @@ namespace UE::Physics
 
 		UE_INTERNAL virtual Chaos::Filter::FInstanceData GetFilterInstanceData() const = 0;
 		UE_INTERNAL virtual void SetFilterInstanceData(const Chaos::Filter::FInstanceData& InInstanceData) = 0;
+		
+		UE_INTERNAL virtual int32 GetNumMaterials() const = 0;
+		UE_INTERNAL virtual FMaterialHandle GetMaterial(const int32 InIndex) const = 0;
+		UE_INTERNAL virtual void SetMaterials(TArray<FMaterialHandle>&& InMaterials) = 0;
+
+		UE_INTERNAL virtual int32 GetNumMaterialMasks() const = 0;
+		UE_INTERNAL virtual FMaterialMaskHandle GetMaterialMask(const int32 InIndex) const = 0;
+		UE_INTERNAL virtual void SetMaterialMasks(TArray<FMaterialMaskHandle>&& InMaterialMasks) = 0;
+
+		UE_INTERNAL virtual int32 GetNumMaterialMaskMaps() const = 0;
+		UE_INTERNAL virtual uint32 GetMaterialMaskMap(const int32 InIndex) const = 0;
+		UE_INTERNAL virtual void SetMaterialMaskMaps(TArray<uint32>&& InMaterialMaskMaps) = 0;
+
+		UE_INTERNAL virtual int32 GetNumMaterialMaskMapMaterials() const = 0;
+		UE_INTERNAL virtual FMaterialHandle GetMaterialMaskMapMaterial(const int32 InIndex) const = 0;
+		UE_INTERNAL virtual void SetMaterialMaskMapMaterials(TArray<FMaterialHandle>&& InMaterialMaskMapMaterials) = 0;
 
 		UE_INTERNAL virtual void* GetUserData() const = 0;
 		UE_INTERNAL virtual void SetUserData(void* InUserData) = 0;
diff --git a/Engine/Source/Runtime/Experimental/RigidPhysics/Public/RigidPhysics/RigidMaterials.h b/Engine/Source/Runtime/Experimental/RigidPhysics/Public/RigidPhysics/RigidMaterials.h
new file mode 100644
--- /dev/null
+++ b/Engine/Source/Runtime/Experimental/RigidPhysics/Public/RigidPhysics/RigidMaterials.h
@@ -0,0 +1,52 @@
+// Copyright Epic Games, Inc. All Rights Reserved.
+#pragma once
+
+#include "RigidPhysics/RigidFwd.h"
+
+#if UE_RIGIDPHYSICS_API_ENABLED
+
+namespace Chaos
+{
+	struct FMaterialHandle;
+	class FChaosPhysicsMaterial;
+	struct FMaterialMaskHandle;
+	class FChaosPhysicsMaterialMask;
+} // namespace Chaos
+
+UE_EXPERIMENTAL(5.8, "The new Chaos API is experimental")
+namespace UE::Physics
+{
+	class FMaterialHandle
+	{
+	public:
+		FMaterialHandle() = default;
+		UE_INTERNAL RIGIDPHYSICS_API FMaterialHandle(const Chaos::FMaterialHandle& InMaterialHandle);
+
+		bool operator==(const FMaterialHandle&) const = default;
+		bool operator!=(const FMaterialHandle&) const = default;
+
+		UE_INTERNAL RIGIDPHYSICS_API Chaos::FMaterialHandle GetMaterialHandle() const;
+		UE_INTERNAL RIGIDPHYSICS_API Chaos::FChaosPhysicsMaterial* GetMaterial() const;
+
+	private:
+		uint32 HandleData = 0;
+	};
+
+	class FMaterialMaskHandle
+	{
+	public:
+		FMaterialMaskHandle() = default;
+		UE_INTERNAL RIGIDPHYSICS_API FMaterialMaskHandle(const Chaos::FMaterialMaskHandle& InMaterialMaskHandle);
+
+		bool operator==(const FMaterialMaskHandle&) const = default;
+		bool operator!=(const FMaterialMaskHandle&) const = default;
+
+		UE_INTERNAL RIGIDPHYSICS_API Chaos::FMaterialMaskHandle GetMaterialMaskHandle() const;
+		UE_INTERNAL RIGIDPHYSICS_API Chaos::FChaosPhysicsMaterialMask* GetMaterialMask() const;
+
+	private:
+		uint32 HandleData = 0;
+	};
+} // namespace UE::Physics
+
+#endif // UE_RIGIDPHYSICS_API_ENABLED
diff --git a/Engine/Source/Runtime/Experimental/RigidPhysics/Public/RigidPhysics/RigidShapeInstance.h b/Engine/Source/Runtime/Experimental/RigidPhysics/Public/RigidPhysics/RigidShapeInstance.h
--- a/Engine/Source/Runtime/Experimental/RigidPhysics/Public/RigidPhysics/RigidShapeInstance.h
+++ b/Engine/Source/Runtime/Experimental/RigidPhysics/Public/RigidPhysics/RigidShapeInstance.h
@@ -130,6 +130,66 @@ namespace UE::Physics
 		{
 			ShapeInstanceRaw->SetFilterInstanceData(InInstanceData);
 		}
+		
+		int32 GetNumMaterials() const
+		{
+			return ShapeInstanceRaw->GetNumMaterials();
+		}
+
+		FMaterialHandle GetMaterial(const int32 InIndex) const
+		{
+			return ShapeInstanceRaw->GetMaterial(InIndex);
+		}
+
+		void SetMaterials(TArray<FMaterialHandle>&& InMaterials) requires CIsGameContext<FContext>
+		{
+			ShapeInstanceRaw->SetMaterials(MoveTemp(InMaterials));
+		}
+
+		int32 GetNumMaterialMasks() const
+		{
+			return ShapeInstanceRaw->GetNumMaterialMasks();
+		}
+
+		FMaterialMaskHandle GetMaterialMask(const int32 InIndex) const
+		{
+			return ShapeInstanceRaw->GetMaterialMask(InIndex);
+		}
+
+		void SetMaterialMasks(TArray<FMaterialMaskHandle>&& InMaterialMasks) requires CIsGameContext<FContext>
+		{
+			ShapeInstanceRaw->SetMaterialMasks(MoveTemp(InMaterialMasks));
+		}
+
+		int32 GetNumMaterialMaskMaps() const
+		{
+			return ShapeInstanceRaw->GetNumMaterialMaskMaps();
+		}
+
+		uint32 GetMaterialMaskMap(const int32 InIndex) const
+		{
+			return ShapeInstanceRaw->GetMaterialMaskMap(InIndex);
+		}
+
+		void SetMaterialMaskMaps(TArray<uint32>&& InMaterialMaskMaps) requires CIsGameContext<FContext>
+		{
+			ShapeInstanceRaw->SetMaterialMaskMaps(MoveTemp(InMaterialMaskMaps));
+		}
+
+		int32 GetNumMaterialMaskMapMaterials() const
+		{
+			return ShapeInstanceRaw->GetNumMaterialMaskMapMaterials();
+		}
+
+		FMaterialHandle GetMaterialMaskMapMaterial(const int32 InIndex) const
+		{
+			return ShapeInstanceRaw->GetMaterialMaskMapMaterial(InIndex);
+		}
+
+		void SetMaterialMaskMapMaterials(TArray<FMaterialHandle>&& InMaterialMaskMapMaterials) requires CIsGameContext<FContext>
+		{
+			ShapeInstanceRaw->SetMaterialMaskMapMaterials(MoveTemp(InMaterialMaskMapMaterials));
+		}
 
 		void* GetUserData() const
 		{
diff --git a/Engine/Source/Runtime/Experimental/RigidPhysics/Public/RigidPhysics/RigidShapeInstanceSetup.h b/Engine/Source/Runtime/Experimental/RigidPhysics/Public/RigidPhysics/RigidShapeInstanceSetup.h
--- a/Engine/Source/Runtime/Experimental/RigidPhysics/Public/RigidPhysics/RigidShapeInstanceSetup.h
+++ b/Engine/Source/Runtime/Experimental/RigidPhysics/Public/RigidPhysics/RigidShapeInstanceSetup.h
@@ -6,6 +6,7 @@
 #if UE_RIGIDPHYSICS_API_ENABLED
 #include "Math/Transform.h"
 #include "RigidPhysics/Geometry/AnyGeometry.h"
+#include "RigidPhysics/RigidMaterials.h"
 #include "BodySetupEnums.h"
 
 UE_EXPERIMENTAL(5.8, "The new Chaos API is experimental")
@@ -21,6 +22,10 @@ namespace UE::Physics
 		Chaos::Filter::FShapeFilterData ShapeFilterData;
 		Chaos::Filter::FInstanceData FilterInstanceData;
 		ECollisionTraceFlag CollisionTraceType = ECollisionTraceFlag::CTF_UseDefault;
+		TArray<FMaterialHandle> Materials;
+		TArray<FMaterialMaskHandle> MaterialMaskHandles;
+		TArray<uint32> MaterialMaskMaps;
+		TArray<FMaterialHandle> MaterialMaskMapMaterials;
 		bool bEnableSim = true;
 		bool bEnableQuery = true;
 		bool bEnableProbe = false;
