Speed up building UAF tests by merging into one target (Part 1)
* This is the first half of the changes to do the following:
* Currently all UAF plugin tests each create their own monolithic builds which takes a while to build.
* These tests cannot reuse prebuilt Editor dlls as they require different defines to function.
* The solution is to have one target that builds all UAF tests and their dependencies once.
* Submit tool now runs all UAF tests when any file inside UAF is modified as before any module's tests were not ran if their dependencies were modified and dependency tracking is hard so we should just run them all, it takes 35 seconds.

#rnx
#Synced-CL 52724535

diff --git a/Engine/Plugins/Experimental/UAF/UAF/Tests/UAFAllTests.Target.cs b/Engine/Plugins/Experimental/UAF/UAF/Tests/UAFAllTests.Target.cs
new file mode 100644
--- /dev/null
+++ b/Engine/Plugins/Experimental/UAF/UAF/Tests/UAFAllTests.Target.cs
@@ -0,0 +1,53 @@
+// Copyright Epic Games, Inc. All Rights Reserved.
+
+using UnrealBuildTool;
+
+[SupportedPlatforms(UnrealPlatformClass.All)]
+public class UAFAllTestsTarget : TestTargetRules
+{
+	public UAFAllTestsTarget(TargetInfo Target) : base(Target)
+	{
+		bUsePlatformFileStub = true;
+		bMockEngineDefaults = true;
+
+		bCompileAgainstEngine = true;
+		bCompileAgainstApplicationCore = true;
+
+		bTestsRequireEngine = true;
+
+		bWarningsAsErrors = true;
+
+		GlobalDefinitions.Add("CATCH_CONFIG_ENABLE_BENCHMARKING=1");
+
+		LaunchModuleName = "UAFTests";
+
+		// Enable all UAF sub-plugins so UBT can discover their test modules.
+		EnablePlugins.AddRange(new string[]
+		{
+			"UAFAnimNode",
+			"UAFAnimGraph",
+			"UAFChooser",
+			"UAFControlRig",
+			"UAFLayering",
+			"UAFMass",
+			"UAFMirroring",
+			"UAFPoseSearch",
+			"UAFStateTree",
+			"UAFWarping",
+		});
+
+		ExtraModuleNames.AddRange(new string[]
+		{
+			"UAFAnimNodeTests",
+			"UAFAnimGraphTests",
+			"UAFChooserTests",
+			"UAFControlRigTests",
+			"UAFLayeringTests",
+			"UAFMassTests",
+			"UAFMirroringTests",
+			"UAFPoseSearchTests",
+			"UAFStateTreeTests",
+			"UAFWarpingTests",
+		});
+	}
+}
